fixed up incorrect/obsolete module dependencies in various Editor modules

diff --git a/Engine/Source/Developer/InternationalizationSettings/InternationalizationSettings.Build.cs b/Engine/Source/Developer/InternationalizationSettings/InternationalizationSettings.Build.cs
--- a/Engine/Source/Developer/InternationalizationSettings/InternationalizationSettings.Build.cs
+++ b/Engine/Source/Developer/InternationalizationSettings/InternationalizationSettings.Build.cs
@@ -7,8 +7,7 @@ namespace UnrealBuildTool.Rules
         public InternationalizationSettings(TargetInfo Target)
 		{
             PrivateDependencyModuleNames.AddRange(
-                new string[]
-                {
+                new string[] {
 				    "Core",
 				    "CoreUObject",
 				    "InputCore",
@@ -17,9 +16,14 @@ namespace UnrealBuildTool.Rules
 				    "EditorStyle",
 				    "PropertyEditor",
 				    "SharedSettingsWidgets",
-				    "Settings",
                 }
-                );
+            );
+
+			PrivateIncludePathModuleNames.AddRange(
+				new string[] {
+					"Settings",
+				}
+			);
 		}
 	}
 }
diff --git a/Engine/Source/Developer/SettingsEditor/SettingsEditor.Build.cs b/Engine/Source/Developer/SettingsEditor/SettingsEditor.Build.cs
--- a/Engine/Source/Developer/SettingsEditor/SettingsEditor.Build.cs
+++ b/Engine/Source/Developer/SettingsEditor/SettingsEditor.Build.cs
@@ -7,17 +7,20 @@ public class SettingsEditor : ModuleRules
 	public SettingsEditor(TargetInfo Target)
 	{
 		PublicDependencyModuleNames.AddRange(
-			new string[]
-			{
+			new string[] {
                 "EditorStyle",
-				"Settings",
 				"Slate",
 			}
 		);
 
-        PrivateDependencyModuleNames.AddRange(
-            new string[]
-            {
+		PublicIncludePathModuleNames.AddRange(
+			new string[] {
+				"Settings",
+			}
+		);
+
+		PrivateDependencyModuleNames.AddRange(
+            new string[] {
 				"Core",
                 "CoreUObject",
                 "InputCore",
@@ -28,8 +31,7 @@ public class SettingsEditor : ModuleRules
         );
 
 		PrivateIncludePaths.AddRange(
-			new string[]
-            {
+			new string[] {
                 "Developer/SettingsEditor/Private",
 				"Developer/SettingsEditor/Private/Models",
                 "Developer/SettingsEditor/Private/Widgets",
diff --git a/Engine/Source/Editor/EditorSettingsViewer/EditorSettingsViewer.Build.cs b/Engine/Source/Editor/EditorSettingsViewer/EditorSettingsViewer.Build.cs
--- a/Engine/Source/Editor/EditorSettingsViewer/EditorSettingsViewer.Build.cs
+++ b/Engine/Source/Editor/EditorSettingsViewer/EditorSettingsViewer.Build.cs
@@ -7,20 +7,17 @@ namespace UnrealBuildTool.Rules
 		public EditorSettingsViewer(TargetInfo Target)
 		{
 			PublicDependencyModuleNames.AddRange(
-				new string[]
-				{
+				new string[] {
 					"Core",
 				}
 			);
 
 			PrivateDependencyModuleNames.AddRange(
-				new string[]
-				{
+				new string[] {
 					"CoreUObject",
 					"Engine",
 					"InputBindingEditor",
 					"MessageLog",
-					"Settings",
 					"SettingsEditor",
 					"Slate",
 					"UnrealEd",
@@ -28,9 +25,14 @@ namespace UnrealBuildTool.Rules
 				}
 			);
 
+			PrivateIncludePathModuleNames.AddRange(
+				new string[] {
+					"Settings",
+				}
+			);
+
 			PrivateIncludePaths.AddRange(
-				new string[]
-				{
+				new string[] {
 					"Editor/EditorSettingsViewer/Private",
 				}
 			);
diff --git a/Engine/Source/Editor/EditorStyle/EditorStyle.Build.cs b/Engine/Source/Editor/EditorStyle/EditorStyle.Build.cs
--- a/Engine/Source/Editor/EditorStyle/EditorStyle.Build.cs
+++ b/Engine/Source/Editor/EditorStyle/EditorStyle.Build.cs
@@ -14,14 +14,14 @@ public class EditorStyle : ModuleRules
 			}
 		);
 
-        PrivateDependencyModuleNames.AddRange(
+		PrivateDependencyModuleNames.Add("TargetPlatform");
+
+		PrivateIncludePathModuleNames.AddRange(
 			new string[] {
 				"Settings",
 			}
 		);
 
-		PrivateDependencyModuleNames.Add("TargetPlatform");
-
 		OptimizeCode = CodeOptimization.Never;
 	}
 }
diff --git a/Engine/Source/Editor/LevelEditor/LevelEditor.Build.cs b/Engine/Source/Editor/LevelEditor/LevelEditor.Build.cs
--- a/Engine/Source/Editor/LevelEditor/LevelEditor.Build.cs
+++ b/Engine/Source/Editor/LevelEditor/LevelEditor.Build.cs
@@ -17,9 +17,10 @@ public class LevelEditor : ModuleRules
 
 		PublicIncludePathModuleNames.AddRange(
 			new string[] {
-				"UserFeedback"
+				"Settings",
+				"UserFeedback",
 			}
-			);
+		);
 
 		PrivateDependencyModuleNames.AddRange(
 			new string[] {
@@ -33,7 +34,6 @@ public class LevelEditor : ModuleRules
 				"Engine",
 				"MessageLog",
 				"NewsFeed",
-				"Settings",
                 "SourceControl",
                 "StatsViewer",
 				"UnrealEd", 
diff --git a/Engine/Source/Editor/NewsFeed/NewsFeed.Build.cs b/Engine/Source/Editor/NewsFeed/NewsFeed.Build.cs
--- a/Engine/Source/Editor/NewsFeed/NewsFeed.Build.cs
+++ b/Engine/Source/Editor/NewsFeed/NewsFeed.Build.cs
@@ -6,19 +6,8 @@ public class NewsFeed : ModuleRules
 {
 	public NewsFeed(TargetInfo Target)
 	{
-		PrivateIncludePaths.AddRange(
-			new string[]
-			{
-				"Editor/NewsFeed/Private",
-				"Editor/NewsFeed/Private/Implementation",
-				"Editor/NewsFeed/Private/Models",
-				"Editor/NewsFeed/Private/Widgets",
-			}
-		);
-
 		PrivateDependencyModuleNames.AddRange(
-			new string[]
-			{
+			new string[] {
 				"Analytics",
 				"Core",
 				"CoreUObject",
@@ -29,9 +18,23 @@ public class NewsFeed : ModuleRules
 				"ImageWrapper",
                 "InputCore",
 				"OnlineSubsystem",
-				"Settings",
 				"Slate",
 			}
 		);
+
+		PrivateIncludePathModuleNames.AddRange(
+			new string[] {
+				"Settings",
+			}
+		);
+
+		PrivateIncludePaths.AddRange(
+			new string[] {
+				"Editor/NewsFeed/Private",
+				"Editor/NewsFeed/Private/Implementation",
+				"Editor/NewsFeed/Private/Models",
+				"Editor/NewsFeed/Private/Widgets",
+			}
+		);
 	}
 }
diff --git a/Engine/Source/Editor/SwarmInterface/SwarmInterface.Build.cs b/Engine/Source/Editor/SwarmInterface/SwarmInterface.Build.cs
--- a/Engine/Source/Editor/SwarmInterface/SwarmInterface.Build.cs
+++ b/Engine/Source/Editor/SwarmInterface/SwarmInterface.Build.cs
@@ -10,12 +10,18 @@ public class SwarmInterface : ModuleRules
 			new string[] {
 				"Editor/SwarmInterface/Public"
 			}
-			);
-		PublicDependencyModuleNames.AddRange( new string[] { "Core", "CoreUObject" } );
+		);
+
+		PublicDependencyModuleNames.AddRange(
+			new string[] {
+				"Core",
+				"CoreUObject",
+			}
+		);
 
 		if (Target.Platform == UnrealTargetPlatform.Mac)
 		{
-			PublicDependencyModuleNames.AddRange( new string[] { "Messaging" } );
+			PrivateIncludePathModuleNames.Add("Messaging");
 		}
 	}
 }
diff --git a/Engine/Source/Editor/TextureEditor/TextureEditor.Build.cs b/Engine/Source/Editor/TextureEditor/TextureEditor.Build.cs
--- a/Engine/Source/Editor/TextureEditor/TextureEditor.Build.cs
+++ b/Engine/Source/Editor/TextureEditor/TextureEditor.Build.cs
@@ -23,7 +23,12 @@ public class TextureEditor : ModuleRules
 			}
 		);
 
-		PrivateIncludePathModuleNames.Add("UnrealEd");
+		PrivateIncludePathModuleNames.AddRange(
+			new string[] {
+				"Settings",
+				"UnrealEd",
+			}
+		);
 
 		PrivateDependencyModuleNames.AddRange(
 			new string[] {
@@ -34,7 +39,6 @@ public class TextureEditor : ModuleRules
 				"RenderCore",
 				"ShaderCore",
 				"RHI",
-				"Settings",
 				"Slate",
                 "EditorStyle",
 				"UnrealEd"
diff --git a/Engine/Source/Runtime/Launch/Launch.Build.cs b/Engine/Source/Runtime/Launch/Launch.Build.cs
--- a/Engine/Source/Runtime/Launch/Launch.Build.cs
+++ b/Engine/Source/Runtime/Launch/Launch.Build.cs
@@ -70,7 +70,7 @@ public class Launch : ModuleRules
 
 		if (UEBuildConfiguration.bCompileAgainstEngine)
 		{
-			PublicDependencyModuleNames.Add("Messaging");
+			PrivateIncludePathModuleNames.Add("Messaging");
 			PublicDependencyModuleNames.Add("SessionServices");
 			PrivateIncludePaths.Add("Developer/DerivedDataCache/Public");
 
