the sections in the Settings UI now automatically sort alphabetically

diff --git a/Engine/Source/Developer/SettingsEditor/Private/Widgets/SSettingsEditor.h b/Engine/Source/Developer/SettingsEditor/Private/Widgets/SSettingsEditor.h
--- a/Engine/Source/Developer/SettingsEditor/Private/Widgets/SSettingsEditor.h
+++ b/Engine/Source/Developer/SettingsEditor/Private/Widgets/SSettingsEditor.h
@@ -90,6 +90,8 @@ private:
 	// Callback for clicking the 'Check Out' button.
 	FReply HandleCheckOutButtonClicked( );
 
+	// Callback for when the user's culture has changed.
+	void HandleCultureChanged( );
 
 	// Callback for determining the visibility of the 'Locked' notice.
 	EVisibility HandleDefaultConfigNoticeVisibility( ) const;
