forgot to check in

diff --git a/Engine/Source/Developer/SettingsEditor/Private/Widgets/SSettingsEditor.cpp b/Engine/Source/Developer/SettingsEditor/Private/Widgets/SSettingsEditor.cpp
--- a/Engine/Source/Developer/SettingsEditor/Private/Widgets/SSettingsEditor.cpp
+++ b/Engine/Source/Developer/SettingsEditor/Private/Widgets/SSettingsEditor.cpp
@@ -17,6 +17,7 @@ SSettingsEditor::~SSettingsEditor( )
 {
 	Model->OnSelectionChanged().RemoveAll(this);
 	SettingsContainer->OnCategoryModified().RemoveAll(this);
+	FCoreDelegates::OnCultureChanged.RemoveAll(this);
 }
 
 
@@ -100,7 +101,7 @@ void SSettingsEditor::Construct( const FArguments& InArgs, const ISettingsEditor
 					.Padding(FMargin(24.0f, 0.0f, 24.0f, 0.0f))
 					[
 						SNew(SSeparator)
-						.Orientation(Orient_Vertical)
+							.Orientation(Orient_Vertical)
 					]
 
 				+ SHorizontalBox::Slot()
@@ -223,6 +224,7 @@ void SSettingsEditor::Construct( const FArguments& InArgs, const ISettingsEditor
 			]
 	];
 
+	FCoreDelegates::OnCultureChanged.AddSP(this, &SSettingsEditor::HandleCultureChanged);
 	Model->OnSelectionChanged().AddSP(this, &SSettingsEditor::HandleModelSelectionChanged);
 	SettingsContainer->OnCategoryModified().AddSP(this, &SSettingsEditor::HandleSettingsContainerCategoryModified);
 
@@ -370,6 +372,27 @@ TSharedRef<SWidget> SSettingsEditor::MakeCategoryWidget( const ISettingsCategory
 		return SNullWidget::NullWidget;
 	}
 
+	// sort the sections alphabetically
+	struct FSectionSortPredicate
+	{
+		FORCEINLINE bool operator()( ISettingsSectionPtr A, ISettingsSectionPtr B ) const
+		{
+			if (!A.IsValid() && !B.IsValid())
+			{
+				return false;
+			}
+
+			if (A.IsValid() != B.IsValid())
+			{
+				return B.IsValid();
+			}
+
+			return (A->GetDisplayName().CompareTo(B->GetDisplayName()) < 0);
+		}
+	};
+
+	Sections.Sort(FSectionSortPredicate());
+
 	// list the sections
 	for (TArray<ISettingsSectionPtr>::TConstIterator It(Sections); It; ++It)
 	{
@@ -475,6 +498,12 @@ FReply SSettingsEditor::HandleCheckOutButtonClicked( )
 }
 
 
+void SSettingsEditor::HandleCultureChanged( )
+{
+	ReloadCategories();
+}
+
+
 EVisibility SSettingsEditor::HandleDefaultConfigNoticeVisibility( ) const
 {
 	TWeakObjectPtr<UObject> SettingsObject = GetSelectedSettingsObject();
