Preferences now inform the user on the sucess of importing and exporting

#ttp 334472 - EDITOR: Feature Request: Preferences: Exporting and Importing could provide more feedback

#branch UE4

diff --git a/Engine/Source/Developer/SettingsEditor/Private/Widgets/SSettingsEditor.cpp b/Engine/Source/Developer/SettingsEditor/Private/Widgets/SSettingsEditor.cpp
--- a/Engine/Source/Developer/SettingsEditor/Private/Widgets/SSettingsEditor.cpp
+++ b/Engine/Source/Developer/SettingsEditor/Private/Widgets/SSettingsEditor.cpp
@@ -536,10 +536,15 @@ FReply SSettingsEditor::HandleExportButtonClicked( )
 
 		if (FDesktopPlatformModule::Get()->SaveFileDialog(ParentWindowHandle, LOCTEXT("ExportSettingsDialogTitle", "Export settings...").ToString(), LastExportDir, DefaultFileName, TEXT("Config files (*.ini)|*.ini"), EFileDialogFlags::None, OutFiles))
 		{
+			FNotificationInfo Notification(LOCTEXT("ExportSettingsFailure", "Export settings failed"));
+			Notification.ExpireDuration = 3.f;
 			if (SelectedSection->Export(OutFiles[0]))
 			{
 				LastExportDir = FPaths::GetPath(OutFiles[0]);
+
+				Notification.Text = FText::Format(LOCTEXT("ExportSettingsSucess", "Settings exported sucessfully to {0}"), FText::FromString(FPaths::GetCleanFilename(OutFiles[0])));
 			}			
+			FSlateNotificationManager::Get().AddNotification(Notification);
 		}
 	}
 
@@ -573,11 +578,15 @@ FReply SSettingsEditor::HandleImportButtonClicked( )
 
 		if (FDesktopPlatformModule::Get()->OpenFileDialog(ParentWindowHandle, LOCTEXT("ImportSettingsDialogTitle", "Import settings...").ToString(), FPaths::GetPath(GEditorUserSettingsIni), TEXT(""), TEXT("Config files (*.ini)|*.ini"), EFileDialogFlags::None, OutFiles))
 		{
-			SelectedSection->Import(OutFiles[0]);
-
 			// Our section has changed but we will not receive a notify from notify hook 
 			// as it was done here and not through user action in the editor.
-			SelectedSection->Save();
+			FNotificationInfo Notification(LOCTEXT("ImportSettingsFailure", "Import settings failed"));
+			Notification.ExpireDuration = 3.f;
+			if (SelectedSection->Import(OutFiles[0]) && SelectedSection->Save())
+			{			
+				Notification.Text = FText::Format(LOCTEXT("ImportSettingsSucess", "Settings imported sucessfully from {0}"), FText::FromString(FPaths::GetCleanFilename(OutFiles[0])));
+			}			
+			FSlateNotificationManager::Get().AddNotification(Notification);
 		}
 	}
 
