SettingsEditor: fixed notification popup always showing error; cosmetic changes

diff --git a/Engine/Source/Developer/SettingsEditor/Private/Widgets/SSettingsEditor.cpp b/Engine/Source/Developer/SettingsEditor/Private/Widgets/SSettingsEditor.cpp
--- a/Engine/Source/Developer/SettingsEditor/Private/Widgets/SSettingsEditor.cpp
+++ b/Engine/Source/Developer/SettingsEditor/Private/Widgets/SSettingsEditor.cpp
@@ -49,8 +49,7 @@ void SSettingsEditor::Construct( const FArguments& InArgs, const ISettingsEditor
 	SettingsView->SetVisibility(TAttribute<EVisibility>::Create(TAttribute<EVisibility>::FGetter::CreateSP(this, &SSettingsEditor::HandleSettingsViewVisibility)));
 	SettingsView->SetIsPropertyEditingEnabledDelegate(FIsPropertyEditingEnabled::CreateSP(this, &SSettingsEditor::HandleSettingsViewEnabled));
 
-	TSharedRef<SWidget> ConfigNoticeWidget =
-		SNew(SSettingsEditorCheckoutNotice)
+	TSharedRef<SWidget> ConfigNoticeWidget = SNew(SSettingsEditorCheckoutNotice)
 		.Visibility(this, &SSettingsEditor::HandleDefaultConfigNoticeVisibility)
 		.Unlocked(this, &SSettingsEditor::IsDefaultConfigEditable)
 		.LockedContent()
@@ -148,10 +147,10 @@ void SSettingsEditor::Construct( const FArguments& InArgs, const ISettingsEditor
 									.VAlign(VAlign_Bottom)
 									.Padding(16.0f, 0.0f, 0.0f, 0.0f)
 									[
-										// save as defaults button
+										// set as default button
 										SNew(SButton)
-											.IsEnabled(this, &SSettingsEditor::HandleSaveDefaultsButtonEnabled)
-											.OnClicked(this, &SSettingsEditor::HandleSaveDefaultsButtonClicked)
+											.IsEnabled(this, &SSettingsEditor::HandleSetAsDefaultButtonEnabled)
+											.OnClicked(this, &SSettingsEditor::HandleSetAsDefaultButtonClicked)
 											.Text(LOCTEXT("SaveDefaultsButtonText", "Set as Default"))
 											.ToolTipText(LOCTEXT("SaveDefaultsButtonTooltip", "Save the values below as the new default settings"))
 									]
@@ -487,6 +486,16 @@ void SSettingsEditor::ReloadCategories(  )
 }
 
 
+void SSettingsEditor::ShowNotification( const FText& Text, SNotificationItem::ECompletionState CompletionState ) const
+{
+	FNotificationInfo Notification(Text);
+	Notification.ExpireDuration = 3.f;
+	Notification.bUseSuccessFailIcons = false;
+
+	FSlateNotificationManager::Get().AddNotification(Notification)->SetCompletionState(CompletionState);
+}
+
+
 /* SSettingsEditor callbacks
  *****************************************************************************/
 
@@ -536,15 +545,14 @@ FReply SSettingsEditor::HandleExportButtonClicked( )
 
 		if (FDesktopPlatformModule::Get()->SaveFileDialog(ParentWindowHandle, LOCTEXT("ExportSettingsDialogTitle", "Export settings...").ToString(), LastExportDir, DefaultFileName, TEXT("Config files (*.ini)|*.ini"), EFileDialogFlags::None, OutFiles))
 		{
-			FNotificationInfo Notification(LOCTEXT("ExportSettingsFailure", "Export settings failed"));
-			Notification.ExpireDuration = 3.f;
 			if (SelectedSection->Export(OutFiles[0]))
 			{
-				LastExportDir = FPaths::GetPath(OutFiles[0]);
-
-				Notification.Text = FText::Format(LOCTEXT("ExportSettingsSucess", "Settings exported sucessfully to {0}"), FText::FromString(FPaths::GetCleanFilename(OutFiles[0])));
-			}			
-			FSlateNotificationManager::Get().AddNotification(Notification);
+				ShowNotification(LOCTEXT("ExportSettingsSuccess", "Export settings succeeded"), SNotificationItem::CS_Success);
+			}
+			else
+			{
+				ShowNotification(LOCTEXT("ExportSettingsFailure", "Export settings failed"), SNotificationItem::CS_Fail);
+			}
 		}
 	}
 
@@ -578,15 +586,14 @@ FReply SSettingsEditor::HandleImportButtonClicked( )
 
 		if (FDesktopPlatformModule::Get()->OpenFileDialog(ParentWindowHandle, LOCTEXT("ImportSettingsDialogTitle", "Import settings...").ToString(), FPaths::GetPath(GEditorUserSettingsIni), TEXT(""), TEXT("Config files (*.ini)|*.ini"), EFileDialogFlags::None, OutFiles))
 		{
-			// Our section has changed but we will not receive a notify from notify hook 
-			// as it was done here and not through user action in the editor.
-			FNotificationInfo Notification(LOCTEXT("ImportSettingsFailure", "Import settings failed"));
-			Notification.ExpireDuration = 3.f;
 			if (SelectedSection->Import(OutFiles[0]) && SelectedSection->Save())
-			{			
-				Notification.Text = FText::Format(LOCTEXT("ImportSettingsSucess", "Settings imported sucessfully from {0}"), FText::FromString(FPaths::GetCleanFilename(OutFiles[0])));
-			}			
-			FSlateNotificationManager::Get().AddNotification(Notification);
+			{
+				ShowNotification(LOCTEXT("ImportSettingsSuccess", "Import settings succeeded"), SNotificationItem::CS_Success);
+			}
+			else
+			{
+				ShowNotification(LOCTEXT("ImportSettingsFailure", "Import settings failed"), SNotificationItem::CS_Fail);
+			}
 		}
 	}
 
@@ -688,7 +695,24 @@ bool SSettingsEditor::HandleResetToDefaultsButtonEnabled( ) const
 }
 
 
-FReply SSettingsEditor::HandleSaveDefaultsButtonClicked( )
+void SSettingsEditor::HandleSectionLinkNavigate( ISettingsSectionPtr Section )
+{
+	Model->SelectSection(Section);
+}
+
+
+EVisibility SSettingsEditor::HandleSectionLinkImageVisibility( ISettingsSectionPtr Section ) const
+{
+	if (Model->GetSelectedSection() == Section)
+	{
+		return EVisibility::Visible;
+	}
+
+	return EVisibility::Hidden;
+}
+
+
+FReply SSettingsEditor::HandleSetAsDefaultButtonClicked( )
 {
 	ISettingsSectionPtr SelectedSection = Model->GetSelectedSection();
 
@@ -720,7 +744,7 @@ FReply SSettingsEditor::HandleSaveDefaultsButtonClicked( )
 }
 
 
-bool SSettingsEditor::HandleSaveDefaultsButtonEnabled( ) const
+bool SSettingsEditor::HandleSetAsDefaultButtonEnabled( ) const
 {
 	ISettingsSectionPtr SelectedSection = Model->GetSelectedSection();
 
@@ -733,23 +757,6 @@ bool SSettingsEditor::HandleSaveDefaultsButtonEnabled( ) const
 }
 
 
-void SSettingsEditor::HandleSectionLinkNavigate( ISettingsSectionPtr Section )
-{
-	Model->SelectSection(Section);
-}
-
-
-EVisibility SSettingsEditor::HandleSectionLinkImageVisibility( ISettingsSectionPtr Section ) const
-{
-	if (Model->GetSelectedSection() == Section)
-	{
-		return EVisibility::Visible;
-	}
-
-	return EVisibility::Hidden;
-}
-
-
 FText SSettingsEditor::HandleSettingsBoxDescriptionText( ) const
 {
 	ISettingsSectionPtr SelectedSection = Model->GetSelectedSection();
diff --git a/Engine/Source/Developer/SettingsEditor/Private/Widgets/SSettingsEditor.h b/Engine/Source/Developer/SettingsEditor/Private/Widgets/SSettingsEditor.h
--- a/Engine/Source/Developer/SettingsEditor/Private/Widgets/SSettingsEditor.h
+++ b/Engine/Source/Developer/SettingsEditor/Private/Widgets/SSettingsEditor.h
@@ -82,6 +82,14 @@ protected:
 	 */
 	void ReloadCategories( );
 
+	/**
+	 * Shows a notification pop-up.
+	 *
+	 * @param Text The notification text.
+	 * @param CompletionState The notification's completion state, i.e. success or failure.
+	 */
+	void ShowNotification( const FText& Text, SNotificationItem::ECompletionState CompletionState ) const;
+
 private:
 
 	// Callback for clicking the Back link.
@@ -120,18 +128,18 @@ private:
 	// Callback for getting the enabled state of the 'Reset to Defaults' button.
 	bool HandleResetToDefaultsButtonEnabled( ) const;
 
-	// Callback for clicking the 'Save as Defaults' button.
-	FReply HandleSaveDefaultsButtonClicked( );
-
-	// Callback for getting the enabled state of the 'Save as Defaults' button.
-	bool HandleSaveDefaultsButtonEnabled( ) const;
-
 	// Callback for navigating a settings section link.
 	void HandleSectionLinkNavigate( ISettingsSectionPtr Section );
 
 	// Callback for getting the visibility of a section link image.
 	EVisibility HandleSectionLinkImageVisibility( ISettingsSectionPtr Section ) const;
 
+	// Callback for clicking the 'Set as Defaults' button.
+	FReply HandleSetAsDefaultButtonClicked( );
+
+	// Callback for getting the enabled state of the 'Set as Defaults' button.
+	bool HandleSetAsDefaultButtonEnabled( ) const;
+
 	// Callback for getting the section description text.
 	FText HandleSettingsBoxDescriptionText( ) const;
 
