SettingsEditor: Ensuring that there is always a settings section selected

diff --git a/Engine/Source/Developer/SettingsEditor/Private/Widgets/SSettingsEditor.cpp b/Engine/Source/Developer/SettingsEditor/Private/Widgets/SSettingsEditor.cpp
--- a/Engine/Source/Developer/SettingsEditor/Private/Widgets/SSettingsEditor.cpp
+++ b/Engine/Source/Developer/SettingsEditor/Private/Widgets/SSettingsEditor.cpp
@@ -51,7 +51,7 @@ void SSettingsEditor::Construct( const FArguments& InArgs, const ISettingsEditor
 
 	TSharedRef<SWidget> ConfigNoticeWidget = SNew(SSettingsEditorCheckoutNotice)
 		.Visibility(this, &SSettingsEditor::HandleDefaultConfigNoticeVisibility)
-		.Unlocked(this, &SSettingsEditor::IsDefaultConfigEditable)
+		.Unlocked(this, &SSettingsEditor::HandleConfigNoticeUnlocked)
 		.LockedContent()
 		[
 			SNew(SButton)
@@ -393,10 +393,8 @@ TSharedRef<SWidget> SSettingsEditor::MakeCategoryWidget( const ISettingsCategory
 	Sections.Sort(FSectionSortPredicate());
 
 	// list the sections
-	for (TArray<ISettingsSectionPtr>::TConstIterator It(Sections); It; ++It)
+	for (const ISettingsSectionPtr Section : Sections)
 	{
-		const ISettingsSectionPtr& Section = *It;
-
 		SectionsBox->AddSlot()
 			.HAlign(HAlign_Left)
 			.Padding(0.0f, VerticalSlotPadding, 0.0f, 0.0f)
@@ -423,6 +421,11 @@ TSharedRef<SWidget> SSettingsEditor::MakeCategoryWidget( const ISettingsCategory
 							.ToolTipText(Section->GetDescription())
 					]
 			];
+
+		if (!Model->GetSelectedSection().IsValid())
+		{
+			Model->SelectSection(Section);
+		}
 	}
 
 	const FSlateBrush* CategoryIcon = FEditorStyle::GetBrush(Category->GetIconName());
@@ -474,13 +477,13 @@ void SSettingsEditor::ReloadCategories(  )
 	TArray<ISettingsCategoryPtr> Categories;
 	SettingsContainer->GetCategories(Categories);
 
-	for (TArray<ISettingsCategoryPtr>::TConstIterator It(Categories); It; ++It)
+	for (const ISettingsCategoryPtr Category : Categories)
 	{
 		CategoriesBox->AddSlot()
 			.AutoHeight()
 			.Padding(0.0f, 0.0f, 0.0f, 8.0f)
 			[
-				MakeCategoryWidget(It->ToSharedRef())
+				MakeCategoryWidget(Category.ToSharedRef())
 			];
 	}
 }
@@ -513,6 +516,12 @@ void SSettingsEditor::HandleCultureChanged( )
 }
 
 
+bool SSettingsEditor::HandleConfigNoticeUnlocked( ) const
+{
+	return !DefaultConfigCheckOutNeeded;
+}
+
+
 EVisibility SSettingsEditor::HandleDefaultConfigNoticeVisibility( ) const
 {
 	TWeakObjectPtr<UObject> SettingsObject = GetSelectedSettingsObject();
@@ -663,12 +672,6 @@ void SSettingsEditor::HandleModelSelectionChanged( )
 }
 
 
-bool SSettingsEditor::IsDefaultConfigEditable( ) const
-{
-	return !DefaultConfigCheckOutNeeded;
-}
-
-
 FReply SSettingsEditor::HandleResetDefaultsButtonClicked( )
 {
 	ISettingsSectionPtr SelectedSection = Model->GetSelectedSection();
diff --git a/Engine/Source/Developer/SettingsEditor/Private/Widgets/SSettingsEditor.h b/Engine/Source/Developer/SettingsEditor/Private/Widgets/SSettingsEditor.h
--- a/Engine/Source/Developer/SettingsEditor/Private/Widgets/SSettingsEditor.h
+++ b/Engine/Source/Developer/SettingsEditor/Private/Widgets/SSettingsEditor.h
@@ -98,6 +98,9 @@ private:
 	// Callback for clicking the 'Check Out' button.
 	FReply HandleCheckOutButtonClicked( );
 
+	// Callback for getting the widget index for the notice switcher.
+	bool HandleConfigNoticeUnlocked( ) const;
+
 	// Callback for when the user's culture has changed.
 	void HandleCultureChanged( );
 
@@ -119,9 +122,6 @@ private:
 	// Callback for changing the selected settings section.
 	void HandleModelSelectionChanged( );
 
-	// Callback for getting the widget index for the notice switcher.
-	bool IsDefaultConfigEditable() const;
-
 	// Callback for clicking the 'Reset to Defaults' button.
 	FReply HandleResetDefaultsButtonClicked( );
 
@@ -163,12 +163,12 @@ private:
 	// Holds the vertical box for settings categories.
 	TSharedPtr<SVerticalBox> CategoriesBox;
 
-	// Holds a timer for checking whether the section's configuration file needs to be checked out.
-	float DefaultConfigCheckOutTimer;
-	
 	// Holds a flag indicating whether the section's configuration file needs to be checked out.
 	bool DefaultConfigCheckOutNeeded;
 
+	// Holds a timer for checking whether the section's configuration file needs to be checked out.
+	float DefaultConfigCheckOutTimer;
+
 	// Holds the overlay slot for custom widgets.
 	SOverlay::FOverlaySlot* CustomWidgetSlot;
 
