added missing documentation; removed redundant line breaks

diff --git a/Engine/Source/Developer/SettingsEditor/Private/Widgets/SSettingsEditorCheckoutNotice.h b/Engine/Source/Developer/SettingsEditor/Private/Widgets/SSettingsEditorCheckoutNotice.h
--- a/Engine/Source/Developer/SettingsEditor/Private/Widgets/SSettingsEditorCheckoutNotice.h
+++ b/Engine/Source/Developer/SettingsEditor/Private/Widgets/SSettingsEditorCheckoutNotice.h
@@ -22,6 +22,11 @@ public:
 
 	SLATE_END_ARGS()
 
+	/**
+	 * Constructs the widget.
+	 *
+	 * @param InArgs The construction arguments.
+	 */
 	void Construct( const FArguments& InArgs );
 
 private:
