documentation fixes

diff --git a/Engine/Source/Developer/SettingsEditor/Private/Models/SettingsEditorModel.h b/Engine/Source/Developer/SettingsEditor/Private/Models/SettingsEditorModel.h
--- a/Engine/Source/Developer/SettingsEditor/Private/Models/SettingsEditorModel.h
+++ b/Engine/Source/Developer/SettingsEditor/Private/Models/SettingsEditorModel.h
@@ -18,7 +18,7 @@ public:
 	/**
 	 * Creates and initializes a new instance.
 	 *
-	 * @param InSettingsContainer - The settings container to use.
+	 * @param InSettingsContainer The settings container to use.
 	 */
 	FSettingsEditorModel( const ISettingsContainerRef& InSettingsContainer )
 		: SettingsContainer(InSettingsContainer)
diff --git a/Engine/Source/Developer/SettingsEditor/Private/Widgets/SSettingsEditor.h b/Engine/Source/Developer/SettingsEditor/Private/Widgets/SSettingsEditor.h
--- a/Engine/Source/Developer/SettingsEditor/Private/Widgets/SSettingsEditor.h
+++ b/Engine/Source/Developer/SettingsEditor/Private/Widgets/SSettingsEditor.h
@@ -31,8 +31,8 @@ public:
 	/**
 	 * Constructs the widget.
 	 *
-	 * @param InArgs - The Slate argument list.
-	 * @param InModel - The view model.
+	 * @param InArgs The Slate argument list.
+	 * @param InModel The view model.
 	 */
 	void Construct( const FArguments& InArgs, const ISettingsEditorModelRef& InModel );
 
@@ -71,7 +71,7 @@ protected:
 	/**
 	 * Creates a widget for the given settings category.
 	 *
-	 * @param Category - The category to create the widget for.
+	 * @param Category The category to create the widget for.
 	 *
 	 * @return The created widget.
 	 */
diff --git a/Engine/Source/Developer/SettingsEditor/Public/Interfaces/ISettingsEditorModel.h b/Engine/Source/Developer/SettingsEditor/Public/Interfaces/ISettingsEditorModel.h
--- a/Engine/Source/Developer/SettingsEditor/Public/Interfaces/ISettingsEditorModel.h
+++ b/Engine/Source/Developer/SettingsEditor/Public/Interfaces/ISettingsEditorModel.h
@@ -42,7 +42,7 @@ public:
 	/**
 	 * Selects the specified settings section to be displayed in the editor.
 	 *
-	 * @param Section - The section to select.
+	 * @param Section The section to select.
 	 *
 	 * @see SelectPreviousSection
 	 */
diff --git a/Engine/Source/Developer/SettingsEditor/Public/Interfaces/ISettingsEditorModule.h b/Engine/Source/Developer/SettingsEditor/Public/Interfaces/ISettingsEditorModule.h
--- a/Engine/Source/Developer/SettingsEditor/Public/Interfaces/ISettingsEditorModule.h
+++ b/Engine/Source/Developer/SettingsEditor/Public/Interfaces/ISettingsEditorModule.h
@@ -18,7 +18,7 @@ public:
 	/**
 	 * Creates a settings editor widget.
 	 *
-	 * @param Model - The view model.
+	 * @param Model The view model.
 	 *
 	 * @return The new widget.
 	 *
@@ -27,9 +27,9 @@ public:
 	virtual TSharedRef<SWidget> CreateEditor( const ISettingsEditorModelRef& Model ) = 0;
 
 	/**
-	 * Creates a controller for the settings editor widget.
+	 * Creates a view model for the settings editor widget.
 	 *
-	 * @param SettingsContainer - The settings container.
+	 * @param SettingsContainer The settings container.
 	 *
 	 * @return The controller.
 	 *
