Settings: Fixed Settings UI incorrectly assumes that default INI files are under source control if they are read-only (TTP# 336653)

diff --git a/Engine/Source/Developer/SettingsEditor/Private/Widgets/SSettingsEditor.cpp b/Engine/Source/Developer/SettingsEditor/Private/Widgets/SSettingsEditor.cpp
--- a/Engine/Source/Developer/SettingsEditor/Private/Widgets/SSettingsEditor.cpp
+++ b/Engine/Source/Developer/SettingsEditor/Private/Widgets/SSettingsEditor.cpp
@@ -49,18 +49,6 @@ void SSettingsEditor::Construct( const FArguments& InArgs, const ISettingsEditor
 	SettingsView->SetVisibility(TAttribute<EVisibility>::Create(TAttribute<EVisibility>::FGetter::CreateSP(this, &SSettingsEditor::HandleSettingsViewVisibility)));
 	SettingsView->SetIsPropertyEditingEnabledDelegate(FIsPropertyEditingEnabled::CreateSP(this, &SSettingsEditor::HandleSettingsViewEnabled));
 
-	TSharedRef<SWidget> ConfigNoticeWidget = SNew(SSettingsEditorCheckoutNotice)
-		.Visibility(this, &SSettingsEditor::HandleDefaultConfigNoticeVisibility)
-		.Unlocked(this, &SSettingsEditor::HandleConfigNoticeUnlocked)
-		.LockedContent()
-		[
-			SNew(SButton)
-			.OnClicked(this, &SSettingsEditor::HandleCheckOutButtonClicked)
-			.Text(LOCTEXT("CheckOutButtonText", "Check Out File"))
-			.ToolTipText(LOCTEXT("CheckOutButtonTooltip", "Check out the default configuration file that holds these settings."))
-			//@TODO: .Visibility(this, &SSettingsFileCheckoutNotice::HandleCheckOutButtonVisibility)
-		];
-
 	ChildSlot
 	[
 		SNew(SVerticalBox)
@@ -202,7 +190,11 @@ void SSettingsEditor::Construct( const FArguments& InArgs, const ISettingsEditor
 							.AutoHeight()
 							.Padding(0.0f, 0.0f, 0.0f, 16.0f)
 							[
-								ConfigNoticeWidget
+								// checkout notice
+								SNew(SSettingsEditorCheckoutNotice)
+									.OnCheckOutClicked(this, &SSettingsEditor::HandleCheckOutButtonClicked)
+									.Visibility(this, &SSettingsEditor::HandleDefaultConfigNoticeVisibility)
+									.Unlocked(this, &SSettingsEditor::HandleConfigNoticeUnlocked)
 							]
 
 						+ SVerticalBox::Slot()
@@ -247,7 +239,7 @@ void SSettingsEditor::Tick( const FGeometry& AllottedGeometry, const double InCu
 
 		if (SettingsObject.IsValid() && SettingsObject->GetClass()->HasAnyClassFlags(CLASS_Config | CLASS_DefaultConfig))
 		{
-			FString ConfigFileName = FString::Printf(TEXT("%sDefault%s.ini"), *FPaths::SourceConfigDir(), *SettingsObject->GetClass()->ClassConfigName.ToString());
+			FString ConfigFileName = GetDefaultConfigFilePath(SettingsObject);
 			bool NewCheckOutNeeded = (FPaths::FileExists(ConfigFileName) && IFileManager::Get().IsReadOnly(*ConfigFileName));
 
 			// file has been checked in or reverted
@@ -287,50 +279,50 @@ void SSettingsEditor::NotifyPostChange( const FPropertyChangedEvent& PropertyCha
 
 bool SSettingsEditor::CheckOutDefaultConfigFile( )
 {
-	FText ErrorMessage;
-
-	// check out configuration file
 	TWeakObjectPtr<UObject> SettingsObject = GetSelectedSettingsObject();
 
-	if (SettingsObject.IsValid())
+	if (!SettingsObject.IsValid())
 	{
-		FString RelativeConfigFilePath = FString::Printf(TEXT("%sDefault%s.ini"), *FPaths::SourceConfigDir(), *SettingsObject->GetClass()->ClassConfigName.ToString());
-		FString AbsoluteConfigFilePath = FPaths::ConvertRelativePathToFull(RelativeConfigFilePath);
+		return false;
+	}
+
+	FText ErrorMessage;
 
-		TArray<FString> FilesToBeCheckedOut;
-		FilesToBeCheckedOut.Add(AbsoluteConfigFilePath);
+	// check out configuration file
+	FString AbsoluteConfigFilePath = GetDefaultConfigFilePath(SettingsObject);
+	TArray<FString> FilesToBeCheckedOut;
+	FilesToBeCheckedOut.Add(AbsoluteConfigFilePath);
 
-		ISourceControlProvider& SourceControlProvider = ISourceControlModule::Get().GetProvider();
-		FSourceControlStatePtr SourceControlState = SourceControlProvider.GetState(AbsoluteConfigFilePath, EStateCacheUsage::ForceUpdate);
+	ISourceControlProvider& SourceControlProvider = ISourceControlModule::Get().GetProvider();
+	FSourceControlStatePtr SourceControlState = SourceControlProvider.GetState(AbsoluteConfigFilePath, EStateCacheUsage::ForceUpdate);
 
-		if(SourceControlState.IsValid())
+	if (SourceControlState.IsValid())
+	{
+		if (SourceControlState->IsDeleted())
 		{
-			if (SourceControlState->IsDeleted())
-			{
-				ErrorMessage = LOCTEXT("ConfigFileMarkedForDeleteError", "Error: The configuration file is marked for deletion.");
-			}
-			else if (!SourceControlState->IsCurrent())
+			ErrorMessage = LOCTEXT("ConfigFileMarkedForDeleteError", "Error: The configuration file is marked for deletion.");
+		}
+		else if (!SourceControlState->IsCurrent())
+		{
+			if (false)
 			{
-				if (false)
+				if (SourceControlProvider.Execute(ISourceControlOperation::Create<FSync>(), FilesToBeCheckedOut) == ECommandResult::Succeeded)
 				{
-					if (SourceControlProvider.Execute(ISourceControlOperation::Create<FSync>(), FilesToBeCheckedOut) == ECommandResult::Succeeded)
-					{
-						SettingsObject->ReloadConfig();
-					}
-					else
-					{
-						ErrorMessage = LOCTEXT("FailedToSyncConfigFileError", "Error: Failed to sync the configuration file to head revision.");
-					}
+					SettingsObject->ReloadConfig();
 				}
-			}
-			else if (SourceControlState->CanCheckout() || SourceControlState->IsCheckedOutOther())
-			{
-				if (SourceControlProvider.Execute(ISourceControlOperation::Create<FCheckOut>(), FilesToBeCheckedOut) == ECommandResult::Failed)
+				else
 				{
-					ErrorMessage = LOCTEXT("FailedToCheckOutConfigFileError", "Error: Failed to check out the configuration file.");
+					ErrorMessage = LOCTEXT("FailedToSyncConfigFileError", "Error: Failed to sync the configuration file to head revision.");
 				}
 			}
 		}
+		else if (SourceControlState->CanCheckout() || SourceControlState->IsCheckedOutOther())
+		{
+			if (SourceControlProvider.Execute(ISourceControlOperation::Create<FCheckOut>(), FilesToBeCheckedOut) == ECommandResult::Failed)
+			{
+				ErrorMessage = LOCTEXT("FailedToCheckOutConfigFileError", "Error: Failed to check out the configuration file.");
+			}
+		}
 	}
 
 	// show errors, if any
@@ -345,6 +337,14 @@ bool SSettingsEditor::CheckOutDefaultConfigFile( )
 }
 
 
+FString SSettingsEditor::GetDefaultConfigFilePath( const TWeakObjectPtr<UObject>& SettingsObject ) const
+{
+	FString RelativeConfigFilePath = FString::Printf(TEXT("%sDefault%s.ini"), *FPaths::SourceConfigDir(), *SettingsObject->GetClass()->ClassConfigName.ToString());
+
+	return FPaths::ConvertRelativePathToFull(RelativeConfigFilePath);
+}
+
+
 TWeakObjectPtr<UObject> SSettingsEditor::GetSelectedSettingsObject( ) const
 {
 	ISettingsSectionPtr SelectedSection = Model->GetSelectedSection();
@@ -470,7 +470,22 @@ TSharedRef<SWidget> SSettingsEditor::MakeCategoryWidget( const ISettingsCategory
 }
 
 
-void SSettingsEditor::ReloadCategories(  )
+bool SSettingsEditor::MakeDefaultConfigFileWritable( )
+{
+	TWeakObjectPtr<UObject> SettingsObject = GetSelectedSettingsObject();
+
+	if (!SettingsObject.IsValid())
+	{
+		return false;
+	}
+
+	FString AbsoluteConfigFilePath = GetDefaultConfigFilePath(SettingsObject);
+
+	return FPlatformFileManager::Get().GetPlatformFile().SetReadOnly(*AbsoluteConfigFilePath, false);
+}
+
+
+void SSettingsEditor::ReloadCategories( )
 {
 	CategoriesBox->ClearChildren();
 
@@ -504,7 +519,14 @@ void SSettingsEditor::ShowNotification( const FText& Text, SNotificationItem::EC
 
 FReply SSettingsEditor::HandleCheckOutButtonClicked( )
 {
-	CheckOutDefaultConfigFile();
+	if (ISourceControlModule::Get().IsEnabled())
+	{
+		CheckOutDefaultConfigFile();
+	}
+	else
+	{
+		MakeDefaultConfigFileWritable();
+	}
 
 	return FReply::Handled();
 }
@@ -734,7 +756,12 @@ FReply SSettingsEditor::HandleSetAsDefaultButtonClicked( )
 			}
 			else
 			{
-				FMessageDialog::Open(EAppMsgType::Ok, LOCTEXT("SaveAsDefaultsIsReadOnlyMessage", "The default configuration file for these settings is currently not writeable, and source control is disabled."));
+				if (FMessageDialog::Open(EAppMsgType::YesNo, LOCTEXT("SaveAsDefaultsIsReadOnlyMessage", "The default configuration file for these settings is currently not writeable. Would you like to make it writable?")) != EAppReturnType::Yes)
+				{
+					return FReply::Handled();
+				}
+
+				MakeDefaultConfigFileWritable();
 			}
 		}
 
diff --git a/Engine/Source/Developer/SettingsEditor/Private/Widgets/SSettingsEditor.h b/Engine/Source/Developer/SettingsEditor/Private/Widgets/SSettingsEditor.h
--- a/Engine/Source/Developer/SettingsEditor/Private/Widgets/SSettingsEditor.h
+++ b/Engine/Source/Developer/SettingsEditor/Private/Widgets/SSettingsEditor.h
@@ -61,6 +61,14 @@ protected:
 	 */
 	bool CheckOutDefaultConfigFile( );
 
+	/**
+	 * Gets the absolute path to the Default.ini for the specified object.
+	 *
+	 * @param SettingsObject The object to get the path for.
+	 * @return The path to the file.
+	 */
+	FString GetDefaultConfigFilePath( const TWeakObjectPtr<UObject>& SettingsObject ) const;
+
 	/**
 	 * Gets the settings object of the selected section, if any.
 	 *
@@ -72,11 +80,17 @@ protected:
 	 * Creates a widget for the given settings category.
 	 *
 	 * @param Category The category to create the widget for.
-	 *
 	 * @return The created widget.
 	 */
 	TSharedRef<SWidget> MakeCategoryWidget( const ISettingsCategoryRef& Category );
 
+	/**
+	 * Makes the default configuration file for the currently selected settings object writable.
+	 *
+	 * @return true if it was made writable, false otherwise.
+	 */
+	bool MakeDefaultConfigFileWritable( );
+
 	/**
 	 * Reloads the settings categories.
 	 */
diff --git a/Engine/Source/Developer/SettingsEditor/Private/Widgets/SSettingsEditorCheckoutNotice.cpp b/Engine/Source/Developer/SettingsEditor/Private/Widgets/SSettingsEditorCheckoutNotice.cpp
--- a/Engine/Source/Developer/SettingsEditor/Private/Widgets/SSettingsEditorCheckoutNotice.cpp
+++ b/Engine/Source/Developer/SettingsEditor/Private/Widgets/SSettingsEditorCheckoutNotice.cpp
@@ -15,6 +15,7 @@
 
 void SSettingsEditorCheckoutNotice::Construct( const FArguments& InArgs )
 {
+	CheckOutClickedDelegate = InArgs._OnCheckOutClicked;
 	bIsUnlocked = InArgs._Unlocked;
 
 	// default configuration notice
@@ -33,7 +34,7 @@ void SSettingsEditorCheckoutNotice::Construct( const FArguments& InArgs )
 					[
 						SNew(SHorizontalBox)
 
-						// Status icon
+						// Locked icon
 						+ SHorizontalBox::Slot()
 							.AutoWidth()
 							.VAlign(VAlign_Center)
@@ -42,28 +43,24 @@ void SSettingsEditorCheckoutNotice::Construct( const FArguments& InArgs )
 									.Image(FEditorStyle::GetBrush("GenericLock"))
 							]
 
-						// Notice
+						// Locked notice
 						+ SHorizontalBox::Slot()
 							.FillWidth(1.0f)
 							.Padding(16.0f, 0.0f)
 							.VAlign(VAlign_Center)
 							[
 								SNew(STextBlock)
-									.Text(LOCTEXT("DefaultSettingsNotice_Source", "These settings are always saved in the default configuration file, which is currently under source control."))
+									.Text(LOCTEXT("DefaultSettingsNotice_Source", "These settings are always saved in the default configuration file, which is currently not writable."))
 							]
 
-						// Call to action
+						// Check out button
 						+ SHorizontalBox::Slot()
 							.AutoWidth()
 							[
-								InArgs._LockedContent.Widget
-							]
-
-						+ SHorizontalBox::Slot()
-							.AutoWidth()
-							[
-								SNew(STextBlock)
-									.Text(this, &SSettingsEditorCheckoutNotice::HandleSccUnavailableTextBlockText)
+								SNew(SButton)
+									.OnClicked(this, &SSettingsEditorCheckoutNotice::HandleCheckOutButtonClicked)
+									.Text(this, &SSettingsEditorCheckoutNotice::HandleCheckOutButtonText)
+									.ToolTipText(this, &SSettingsEditorCheckoutNotice::HandleCheckOutButtonToolTip)
 							]
 					]
 
@@ -72,6 +69,7 @@ void SSettingsEditorCheckoutNotice::Construct( const FArguments& InArgs )
 					[
 						SNew(SHorizontalBox)
 
+						// Unlocked icon
 						+ SHorizontalBox::Slot()
 							.AutoWidth()
 							.VAlign(VAlign_Center)
@@ -80,6 +78,7 @@ void SSettingsEditorCheckoutNotice::Construct( const FArguments& InArgs )
 									.Image(FEditorStyle::GetBrush("GenericUnlock"))
 							]
 
+						// Unlocked notice
 						+ SHorizontalBox::Slot()
 							.AutoWidth()
 							.Padding(16.0f, 0.0f)
@@ -97,30 +96,47 @@ void SSettingsEditorCheckoutNotice::Construct( const FArguments& InArgs )
 /* SSettingsEditorCheckoutNotice callbacks
  *****************************************************************************/
 
-EVisibility SSettingsEditorCheckoutNotice::HandleCheckOutButtonVisibility( ) const
+FReply SSettingsEditorCheckoutNotice::HandleCheckOutButtonClicked( )
+{
+	if (CheckOutClickedDelegate.IsBound())
+	{
+		return CheckOutClickedDelegate.Execute();
+	}
+
+	return FReply::Handled();
+}
+
+
+FText SSettingsEditorCheckoutNotice::HandleCheckOutButtonText( ) const
 {
 	if (ISourceControlModule::Get().IsEnabled() && ISourceControlModule::Get().GetProvider().IsAvailable())
 	{
-		return EVisibility::Visible;
+		return LOCTEXT("CheckOutFile", "Check Out File");
 	}
 
-	return EVisibility::Collapsed;
+	return LOCTEXT("MakeWritable", "Make Writable");
 }
 
 
-FText SSettingsEditorCheckoutNotice::HandleSccUnavailableTextBlockText( ) const
+FText SSettingsEditorCheckoutNotice::HandleCheckOutButtonToolTip( ) const
 {
-	if (!ISourceControlModule::Get().IsEnabled())
+	if (ISourceControlModule::Get().IsEnabled() && ISourceControlModule::Get().GetProvider().IsAvailable())
 	{
-		return LOCTEXT("SccDisabledNotice", "[Source control is disabled]");
+		return LOCTEXT("CheckOutFileTooltip", "Check out the default configuration file that holds these settings.");
 	}
 
-	if (!ISourceControlModule::Get().GetProvider().IsAvailable())
+	return LOCTEXT("MakeWritableTooltip", "Make the default configuration file that holds these settings writable.");
+}
+
+
+EVisibility SSettingsEditorCheckoutNotice::HandleCheckOutButtonVisibility( ) const
+{
+	if (ISourceControlModule::Get().IsEnabled() && ISourceControlModule::Get().GetProvider().IsAvailable())
 	{
-		return LOCTEXT("SccUnavailableNotice", "[Source control server unavailable]");
+		return EVisibility::Visible;
 	}
 
-	return FText::GetEmpty();
+	return EVisibility::Collapsed;
 }
 
 
diff --git a/Engine/Source/Developer/SettingsEditor/Private/Widgets/SSettingsEditorCheckoutNotice.h b/Engine/Source/Developer/SettingsEditor/Private/Widgets/SSettingsEditorCheckoutNotice.h
--- a/Engine/Source/Developer/SettingsEditor/Private/Widgets/SSettingsEditorCheckoutNotice.h
+++ b/Engine/Source/Developer/SettingsEditor/Private/Widgets/SSettingsEditorCheckoutNotice.h
@@ -20,6 +20,9 @@ public:
 		/** Slot for this button's content (optional) */
 		SLATE_NAMED_SLOT(FArguments, LockedContent)
 
+		/** Called when the 'Check Out' button is clicked */
+		SLATE_EVENT(FOnClicked, OnCheckOutClicked)
+
 	SLATE_END_ARGS()
 
 	/**
@@ -31,6 +34,15 @@ public:
 
 private:
 
+	// Callback for clicking the 'Check Out' button.
+	FReply HandleCheckOutButtonClicked( );
+
+	// Callback for getting the text of the 'Check Out' button.
+	FText HandleCheckOutButtonText( ) const;
+
+	// Callback for getting the tool tip text of the 'Check Out' button.
+	FText HandleCheckOutButtonToolTip( ) const;
+
 	// Callback for determining the visibility of the check-out button.
 	EVisibility HandleCheckOutButtonVisibility( ) const;
 
@@ -40,10 +52,10 @@ private:
 		return bIsUnlocked.Get() ? 1 : 0;
 	}
 
-	// Callback for getting the visibility of the 'Source control unavailable' text block.
-	FText HandleSccUnavailableTextBlockText( ) const;
-
 private:
 
+	// Holds a delegate that is executed when the 'Check Out' button has been clicked.
+	FOnClicked CheckOutClickedDelegate;
+
 	TAttribute<bool> bIsUnlocked;
 };
