SettingsEditor: Documentation cleanup pass

diff --git a/Engine/Source/Developer/SettingsEditor/Private/Widgets/SSettingsEditorCheckoutNotice.h b/Engine/Source/Developer/SettingsEditor/Private/Widgets/SSettingsEditorCheckoutNotice.h
--- a/Engine/Source/Developer/SettingsEditor/Private/Widgets/SSettingsEditorCheckoutNotice.h
+++ b/Engine/Source/Developer/SettingsEditor/Private/Widgets/SSettingsEditorCheckoutNotice.h
@@ -4,6 +4,9 @@
 	SSettingsEditorCheckoutNotice.h: Declares the SSettingsEditorCheckoutNotice class.
 =============================================================================*/
 
+#pragma once
+
+
 /**
  * Implements a widget that provides a notice for files that need to be checked out.
  */
diff --git a/Engine/Source/Developer/SettingsEditor/Public/Interfaces/ISettingsEditorModel.h b/Engine/Source/Developer/SettingsEditor/Public/Interfaces/ISettingsEditorModel.h
--- a/Engine/Source/Developer/SettingsEditor/Public/Interfaces/ISettingsEditorModel.h
+++ b/Engine/Source/Developer/SettingsEditor/Public/Interfaces/ISettingsEditorModel.h
@@ -7,19 +7,20 @@
 #pragma once
 
 
-/**
- * Type definition for shared pointers to instances of ISettingsEditorModel.
- */
+/** Type definition for shared pointers to instances of ISettingsEditorModel. */
 typedef TSharedPtr<class ISettingsEditorModel> ISettingsEditorModelPtr;
 
-/**
- * Type definition for shared references to instances of ISettingsEditorModel.
- */
+/** Type definition for shared references to instances of ISettingsEditorModel. */
 typedef TSharedRef<class ISettingsEditorModel> ISettingsEditorModelRef;
 
 
 /**
  * Interface for settings editor view models.
+ *
+ * The settings editor view model stores the view state for the Settings Editor UI.
+ * Instances of this interface can be passed to Settings Editors in order to provide
+ * access to the settings container being added and to perform various user actions,
+ * such as setting the currently selected settings section.
  */
 class ISettingsEditorModel
 {
diff --git a/Engine/Source/Developer/SettingsEditor/Public/Interfaces/ISettingsEditorModule.h b/Engine/Source/Developer/SettingsEditor/Public/Interfaces/ISettingsEditorModule.h
--- a/Engine/Source/Developer/SettingsEditor/Public/Interfaces/ISettingsEditorModule.h
+++ b/Engine/Source/Developer/SettingsEditor/Public/Interfaces/ISettingsEditorModule.h
@@ -19,9 +19,7 @@ public:
 	 * Creates a settings editor widget.
 	 *
 	 * @param Model The view model.
-	 *
 	 * @return The new widget.
-	 *
 	 * @see CreateController
 	 */
 	virtual TSharedRef<SWidget> CreateEditor( const ISettingsEditorModelRef& Model ) = 0;
@@ -30,9 +28,7 @@ public:
 	 * Creates a view model for the settings editor widget.
 	 *
 	 * @param SettingsContainer The settings container.
-	 *
 	 * @return The controller.
-	 *
 	 * @see CreateEditor
 	 */
 	virtual ISettingsEditorModelRef CreateModel( const ISettingsContainerRef& SettingsContainer ) = 0;
