Settings editor no longer relies on read-only state of ini files.

The settings editor now queries the appropriate source control state of files that it its editing, rather than using the read-only state (which is only applicable to Perforce).
Also added a throbber to indicate that a source control query is in progress.

TTP# 337061 - EDITOR: SCC: SETTINGS: Settings Editor may not work correctly with some SCC providers

diff --git a/Engine/Source/Developer/SettingsEditor/Private/Widgets/SSettingsEditor.cpp b/Engine/Source/Developer/SettingsEditor/Private/Widgets/SSettingsEditor.cpp
--- a/Engine/Source/Developer/SettingsEditor/Private/Widgets/SSettingsEditor.cpp
+++ b/Engine/Source/Developer/SettingsEditor/Private/Widgets/SSettingsEditor.cpp
@@ -25,6 +25,7 @@ void SSettingsEditor::Construct( const FArguments& InArgs, const ISettingsEditor
 	Model = InModel;
 	DefaultConfigCheckOutTimer = 0.0f;
 	DefaultConfigCheckOutNeeded = false;
+	DefaultConfigQueryInProgress = false;
 	SettingsContainer = InModel->GetSettingsContainer();
 
 	// initialize settings view
@@ -191,6 +192,7 @@ void SSettingsEditor::Construct( const FArguments& InArgs, const ISettingsEditor
 									.OnCheckOutClicked(this, &SSettingsEditor::HandleCheckOutButtonClicked)
 									.Visibility(this, &SSettingsEditor::HandleDefaultConfigNoticeVisibility)
 									.Unlocked(this, &SSettingsEditor::HandleConfigNoticeUnlocked)
+									.LookingForSourceControlState(this, &SSettingsEditor::HandleLookingForSourceControlState)
 							]
 
 						+ SVerticalBox::Slot()
@@ -235,8 +237,22 @@ void SSettingsEditor::Tick( const FGeometry& AllottedGeometry, const double InCu
 
 		if (SettingsObject.IsValid() && SettingsObject->GetClass()->HasAnyClassFlags(CLASS_Config | CLASS_DefaultConfig))
 		{
-			FString ConfigFileName = GetDefaultConfigFilePath(SettingsObject);
-			bool NewCheckOutNeeded = (FPaths::FileExists(ConfigFileName) && IFileManager::Get().IsReadOnly(*ConfigFileName));
+			CachedConfigFileName = GetDefaultConfigFilePath(SettingsObject);
+			bool NewCheckOutNeeded = false;
+			if(ISourceControlModule::Get().IsEnabled())
+			{
+				// note: calling QueueStatusUpdate often does not spam status updates as an internal timer prevents this
+				ISourceControlModule::Get().QueueStatusUpdate(CachedConfigFileName);
+
+				ISourceControlProvider& SourceControlProvider = ISourceControlModule::Get().GetProvider();
+				FSourceControlStatePtr SourceControlState = SourceControlProvider.GetState(CachedConfigFileName, EStateCacheUsage::Use);
+				NewCheckOutNeeded = SourceControlState.IsValid() && SourceControlState->CanCheckout();
+				DefaultConfigQueryInProgress = SourceControlState.IsValid() && SourceControlState->IsUnknown();
+			}
+			else
+			{
+				NewCheckOutNeeded = (FPaths::FileExists(CachedConfigFileName) && IFileManager::Get().IsReadOnly(*CachedConfigFileName));
+			}
 
 			// file has been checked in or reverted
 			if ((NewCheckOutNeeded == true) && (DefaultConfigCheckOutNeeded == false))
@@ -536,9 +552,13 @@ void SSettingsEditor::HandleCultureChanged( )
 
 bool SSettingsEditor::HandleConfigNoticeUnlocked( ) const
 {
-	return !DefaultConfigCheckOutNeeded;
+	return !DefaultConfigCheckOutNeeded && !DefaultConfigQueryInProgress;
 }
 
+bool SSettingsEditor::HandleLookingForSourceControlState() const
+{
+	return DefaultConfigQueryInProgress;
+}
 
 EVisibility SSettingsEditor::HandleDefaultConfigNoticeVisibility( ) const
 {
@@ -743,7 +763,7 @@ FReply SSettingsEditor::HandleSetAsDefaultButtonClicked( )
 		{
 			if (ISourceControlModule::Get().IsEnabled())
 			{
-				if (FMessageDialog::Open(EAppMsgType::YesNo, LOCTEXT("SaveAsDefaultNeedsCheckoutMessage", "The default configuration file for these settings is currently not writeable. Would you like to check it out from source control?")) != EAppReturnType::Yes)
+				if (FMessageDialog::Open(EAppMsgType::YesNo, LOCTEXT("SaveAsDefaultNeedsCheckoutMessage", "The default configuration file for these settings is currently not checked out. Would you like to check it out from source control?")) != EAppReturnType::Yes)
 				{
 					return FReply::Handled();
 				}
@@ -832,7 +852,7 @@ bool SSettingsEditor::HandleSettingsViewEnabled( ) const
 {
 	ISettingsSectionPtr SelectedSection = Model->GetSelectedSection();
 
-	return (SelectedSection.IsValid() && SelectedSection->CanEdit() && (!SelectedSection->HasDefaultSettingsObject() || !DefaultConfigCheckOutNeeded));
+	return (SelectedSection.IsValid() && SelectedSection->CanEdit() && (!SelectedSection->HasDefaultSettingsObject() || (!DefaultConfigCheckOutNeeded && !DefaultConfigQueryInProgress)));
 }
 
 
diff --git a/Engine/Source/Developer/SettingsEditor/Private/Widgets/SSettingsEditor.h b/Engine/Source/Developer/SettingsEditor/Private/Widgets/SSettingsEditor.h
--- a/Engine/Source/Developer/SettingsEditor/Private/Widgets/SSettingsEditor.h
+++ b/Engine/Source/Developer/SettingsEditor/Private/Widgets/SSettingsEditor.h
@@ -164,6 +164,9 @@ private:
 	// Callback for determining the visibility of the settings view.
 	EVisibility HandleSettingsViewVisibility( ) const;
 
+	// Callback for determining whether we are currently looking for the source control state of the config file
+	bool HandleLookingForSourceControlState( ) const;
+
 private:
 
 	// Holds the vertical box for settings categories.
@@ -172,6 +175,9 @@ private:
 	// Holds a flag indicating whether the section's configuration file needs to be checked out.
 	bool DefaultConfigCheckOutNeeded;
 
+	// Holds a flag indicating whether the section's configuration file has a source control status request in progress.
+	bool DefaultConfigQueryInProgress;
+
 	// Holds a timer for checking whether the section's configuration file needs to be checked out.
 	float DefaultConfigCheckOutTimer;
 
@@ -189,4 +195,7 @@ private:
 
 	// Holds the details view.
 	TSharedPtr<IDetailsView> SettingsView;
+
+	// Cached filename for the config file we are displaying
+	FString CachedConfigFileName;
 };
diff --git a/Engine/Source/Developer/SettingsEditor/Private/Widgets/SSettingsEditorCheckoutNotice.cpp b/Engine/Source/Developer/SettingsEditor/Private/Widgets/SSettingsEditorCheckoutNotice.cpp
--- a/Engine/Source/Developer/SettingsEditor/Private/Widgets/SSettingsEditorCheckoutNotice.cpp
+++ b/Engine/Source/Developer/SettingsEditor/Private/Widgets/SSettingsEditorCheckoutNotice.cpp
@@ -13,6 +13,7 @@ void SSettingsEditorCheckoutNotice::Construct( const FArguments& InArgs )
 {
 	CheckOutClickedDelegate = InArgs._OnCheckOutClicked;
 	bIsUnlocked = InArgs._Unlocked;
+	bLookingForSourceControlState = InArgs._LookingForSourceControlState;
 
 	// default configuration notice
 	ChildSlot
@@ -46,7 +47,7 @@ void SSettingsEditorCheckoutNotice::Construct( const FArguments& InArgs )
 							.VAlign(VAlign_Center)
 							[
 								SNew(STextBlock)
-									.Text(LOCTEXT("DefaultSettingsNotice_Source", "These settings are always saved in the default configuration file, which is currently not writable."))
+									.Text(this, &SSettingsEditorCheckoutNotice::HandleLockedStatusText)
 							]
 
 						// Check out button
@@ -57,6 +58,16 @@ void SSettingsEditorCheckoutNotice::Construct( const FArguments& InArgs )
 									.OnClicked(this, &SSettingsEditorCheckoutNotice::HandleCheckOutButtonClicked)
 									.Text(this, &SSettingsEditorCheckoutNotice::HandleCheckOutButtonText)
 									.ToolTipText(this, &SSettingsEditorCheckoutNotice::HandleCheckOutButtonToolTip)
+									.Visibility(this, &SSettingsEditorCheckoutNotice::HandleCheckOutButtonVisibility)
+							]
+
+						// Source control status throbber
+						+ SHorizontalBox::Slot()
+							.AutoWidth()
+							.Padding(0.0f, 1.0f)
+							[
+								SNew(SThrobber)
+									.Visibility(this, &SSettingsEditorCheckoutNotice::HandleThrobberVisibility)
 							]
 					]
 
@@ -81,7 +92,7 @@ void SSettingsEditorCheckoutNotice::Construct( const FArguments& InArgs )
 							.VAlign(VAlign_Center)
 							[
 								SNew(STextBlock)
-									.Text(LOCTEXT("DefaultSettingsNotice_Writable", "These settings are always saved in the default configuration file, which is currently writable."))
+									.Text(this, &SSettingsEditorCheckoutNotice::HandleUnlockedStatusText)
 							]
 					]
 			]
@@ -129,11 +140,33 @@ EVisibility SSettingsEditorCheckoutNotice::HandleCheckOutButtonVisibility( ) con
 {
 	if (ISourceControlModule::Get().IsEnabled() && ISourceControlModule::Get().GetProvider().IsAvailable())
 	{
-		return EVisibility::Visible;
+		return !bLookingForSourceControlState.Get() ? EVisibility::Visible : EVisibility::Collapsed;
 	}
 
 	return EVisibility::Collapsed;
 }
 
+FText SSettingsEditorCheckoutNotice::HandleLockedStatusText() const
+{
+	return ISourceControlModule::Get().IsEnabled() ? 
+		LOCTEXT("DefaultSettingsNotice_WithSourceControl", "These settings are always saved in the default configuration file, which is currently not checked out.") :
+		LOCTEXT("DefaultSettingsNotice_Source", "These settings are always saved in the default configuration file, which is currently not writable.");
+}
+
+FText SSettingsEditorCheckoutNotice::HandleUnlockedStatusText() const
+{
+	return ISourceControlModule::Get().IsEnabled() ? 
+		LOCTEXT("DefaultSettingsNotice_CheckedOut", "These settings are always saved in the default configuration file, which is currently checked out.") :
+		LOCTEXT("DefaultSettingsNotice_Writable", "These settings are always saved in the default configuration file, which is currently writable.");
+}
+
+EVisibility SSettingsEditorCheckoutNotice::HandleThrobberVisibility() const
+{
+	if(ISourceControlModule::Get().IsEnabled())
+	{
+		return bLookingForSourceControlState.Get() ? EVisibility::Visible : EVisibility::Collapsed;
+	}
+	return EVisibility::Collapsed;
+}
 
 #undef LOCTEXT_NAMESPACE
diff --git a/Engine/Source/Developer/SettingsEditor/Private/Widgets/SSettingsEditorCheckoutNotice.h b/Engine/Source/Developer/SettingsEditor/Private/Widgets/SSettingsEditorCheckoutNotice.h
--- a/Engine/Source/Developer/SettingsEditor/Private/Widgets/SSettingsEditorCheckoutNotice.h
+++ b/Engine/Source/Developer/SettingsEditor/Private/Widgets/SSettingsEditorCheckoutNotice.h
@@ -22,6 +22,9 @@ public:
 		/** Called when the 'Check Out' button is clicked */
 		SLATE_EVENT(FOnClicked, OnCheckOutClicked)
 
+		/** Called to determine if we are currently looking for the source control state of the file */
+		SLATE_ATTRIBUTE(bool, LookingForSourceControlState)
+
 	SLATE_END_ARGS()
 
 	/**
@@ -51,10 +54,21 @@ private:
 		return bIsUnlocked.Get() ? 1 : 0;
 	}
 
+	// Callback for getting the status text when the config is locked
+	FText HandleLockedStatusText() const;
+
+	// Callback for getting the status text when the config is unlocked
+	FText HandleUnlockedStatusText() const;
+
+	// Callback for getting the visibility of the source control throbber
+	EVisibility HandleThrobberVisibility() const;
+
 private:
 
 	// Holds a delegate that is executed when the 'Check Out' button has been clicked.
 	FOnClicked CheckOutClickedDelegate;
 
 	TAttribute<bool> bIsUnlocked;
+
+	TAttribute<bool> bLookingForSourceControlState;
 };
