Record when a preference is changed for analytics.  The value of the field is not sent.

diff --git a/Engine/Source/Developer/SettingsEditor/Private/Widgets/SSettingsEditor.cpp b/Engine/Source/Developer/SettingsEditor/Private/Widgets/SSettingsEditor.cpp
--- a/Engine/Source/Developer/SettingsEditor/Private/Widgets/SSettingsEditor.cpp
+++ b/Engine/Source/Developer/SettingsEditor/Private/Widgets/SSettingsEditor.cpp
@@ -1,7 +1,9 @@
 // Copyright 1998-2014 Epic Games, Inc. All Rights Reserved.
 
 #include "SettingsEditorPrivatePCH.h"
-
+#include "EngineAnalytics.h"
+#include "AnalyticsEventAttribute.h"
+#include "IAnalyticsProvider.h"
 
 #define LOCTEXT_NAMESPACE "SSettingsEditor"
 
@@ -281,6 +283,8 @@ void SSettingsEditor::NotifyPostChange( const FPropertyChangedEvent& PropertyCha
 
 	if (SelectedSection.IsValid() && (PropertyChangedEvent.ChangeType != EPropertyChangeType::Interactive))
 	{
+		RecordPreferenceChangedAnalytics( SelectedSection, PropertyChangedEvent );
+
 		SelectedSection->Save();
 	}
 }
@@ -560,6 +564,21 @@ bool SSettingsEditor::HandleLookingForSourceControlState() const
 	return DefaultConfigQueryInProgress;
 }
 
+void SSettingsEditor::RecordPreferenceChangedAnalytics( ISettingsSectionPtr SelectedSection, const FPropertyChangedEvent& PropertyChangedEvent ) const
+{
+	UProperty* ChangedProperty = PropertyChangedEvent.MemberProperty;
+	// submit analytics data
+	if(FEngineAnalytics::IsAvailable() && ChangedProperty != nullptr)
+	{
+		TArray<FAnalyticsEventAttribute> EventAttributes;
+		EventAttributes.Add(FAnalyticsEventAttribute(TEXT("PropertySection"), SelectedSection->GetDisplayName().ToString()));
+		EventAttributes.Add(FAnalyticsEventAttribute(TEXT("PropertyClass"), ChangedProperty->GetOwnerClass()->GetName()));
+		EventAttributes.Add(FAnalyticsEventAttribute(TEXT("PropertyName"), ChangedProperty->GetName()));
+
+		FEngineAnalytics::GetProvider().RecordEvent(TEXT("Editor.Usage.PreferencesChanged"), EventAttributes);
+	}
+}
+
 EVisibility SSettingsEditor::HandleDefaultConfigNoticeVisibility( ) const
 {
 	TWeakObjectPtr<UObject> SettingsObject = GetSelectedSettingsObject();
diff --git a/Engine/Source/Developer/SettingsEditor/Private/Widgets/SSettingsEditor.h b/Engine/Source/Developer/SettingsEditor/Private/Widgets/SSettingsEditor.h
--- a/Engine/Source/Developer/SettingsEditor/Private/Widgets/SSettingsEditor.h
+++ b/Engine/Source/Developer/SettingsEditor/Private/Widgets/SSettingsEditor.h
@@ -167,6 +167,8 @@ private:
 	// Callback for determining whether we are currently looking for the source control state of the config file
 	bool HandleLookingForSourceControlState( ) const;
 
+	// Reports a preference changed event to the analytics system */
+	void RecordPreferenceChangedAnalytics( ISettingsSectionPtr SelectedSection, const FPropertyChangedEvent& PropertyChangedEvent ) const;
 private:
 
 	// Holds the vertical box for settings categories.
diff --git a/Engine/Source/Developer/SettingsEditor/SettingsEditor.Build.cs b/Engine/Source/Developer/SettingsEditor/SettingsEditor.Build.cs
--- a/Engine/Source/Developer/SettingsEditor/SettingsEditor.Build.cs
+++ b/Engine/Source/Developer/SettingsEditor/SettingsEditor.Build.cs
@@ -28,6 +28,8 @@ public class SettingsEditor : ModuleRules
 				"PropertyEditor",
 				"SlateCore",
 				"SourceControl",
+				"Analytics",
+				"Engine",
             }
         );
 
