Settings: layout tweak

diff --git a/Engine/Source/Developer/SettingsEditor/Private/Widgets/SSettingsEditor.cpp b/Engine/Source/Developer/SettingsEditor/Private/Widgets/SSettingsEditor.cpp
--- a/Engine/Source/Developer/SettingsEditor/Private/Widgets/SSettingsEditor.cpp
+++ b/Engine/Source/Developer/SettingsEditor/Private/Widgets/SSettingsEditor.cpp
@@ -201,7 +201,7 @@ void SSettingsEditor::Construct( const FArguments& InArgs, const ISettingsEditor
 
 						+ SVerticalBox::Slot()
 							.FillHeight(1.0f)
-							.Padding(0.0f, 12.0f, 0.0f, 0.0f)
+							.Padding(0.0f, 40.0f, 0.0f, 0.0f)
 							[
 								// settings area
 								SNew(SOverlay)
