Editor: Some small improvements to project settings dialog
- Change checkout notice for settings ini files to go green when they are editable
- Reworded messages to include the actual ini file name
- Hid 'reset to default' and 'set as default' options when editing a default ini file, since they make no sense in this context and were always disabled

diff --git a/Engine/Source/Developer/SettingsEditor/Private/Widgets/SSettingsEditor.cpp b/Engine/Source/Developer/SettingsEditor/Private/Widgets/SSettingsEditor.cpp
--- a/Engine/Source/Developer/SettingsEditor/Private/Widgets/SSettingsEditor.cpp
+++ b/Engine/Source/Developer/SettingsEditor/Private/Widgets/SSettingsEditor.cpp
@@ -138,6 +138,7 @@ void SSettingsEditor::Construct( const FArguments& InArgs, const ISettingsEditor
 									[
 										// set as default button
 										SNew(SButton)
+											.Visibility(this, &SSettingsEditor::EditingNonDefaultSettingsVisibility)
 											.IsEnabled(this, &SSettingsEditor::HandleSetAsDefaultButtonEnabled)
 											.OnClicked(this, &SSettingsEditor::HandleSetAsDefaultButtonClicked)
 											.Text(LOCTEXT("SaveDefaultsButtonText", "Set as Default"))
@@ -180,6 +181,7 @@ void SSettingsEditor::Construct( const FArguments& InArgs, const ISettingsEditor
 									[
 										// reset defaults button
 										SNew(SButton)
+											.Visibility(this, &SSettingsEditor::EditingNonDefaultSettingsVisibility)
 											.IsEnabled(this, &SSettingsEditor::HandleResetToDefaultsButtonEnabled)
 											.OnClicked(this, &SSettingsEditor::HandleResetDefaultsButtonClicked)
 											.Text(LOCTEXT("ResetDefaultsButtonText", "Reset to Defaults"))
@@ -196,6 +198,7 @@ void SSettingsEditor::Construct( const FArguments& InArgs, const ISettingsEditor
 									.OnCheckOutClicked(this, &SSettingsEditor::HandleCheckOutButtonClicked)
 									.Visibility(this, &SSettingsEditor::HandleDefaultConfigNoticeVisibility)
 									.Unlocked(this, &SSettingsEditor::HandleConfigNoticeUnlocked)
+									.ConfigFilePath(this, &SSettingsEditor::GetCachedConfigFileName)
 									.LookingForSourceControlState(this, &SSettingsEditor::HandleLookingForSourceControlState)
 							]
 
@@ -735,6 +738,14 @@ bool SSettingsEditor::HandleResetToDefaultsButtonEnabled( ) const
 }
 
 
+EVisibility SSettingsEditor::EditingNonDefaultSettingsVisibility( ) const
+{
+	ISettingsSectionPtr SelectedSection = Model->GetSelectedSection();
+
+	return (SelectedSection.IsValid() && SelectedSection->HasDefaultSettingsObject()) ? EVisibility::Collapsed : EVisibility::Visible;
+}
+
+
 void SSettingsEditor::HandleSectionLinkNavigate( ISettingsSectionPtr Section )
 {
 	Model->SelectSection(Section);
diff --git a/Engine/Source/Developer/SettingsEditor/Private/Widgets/SSettingsEditor.h b/Engine/Source/Developer/SettingsEditor/Private/Widgets/SSettingsEditor.h
--- a/Engine/Source/Developer/SettingsEditor/Private/Widgets/SSettingsEditor.h
+++ b/Engine/Source/Developer/SettingsEditor/Private/Widgets/SSettingsEditor.h
@@ -134,6 +134,9 @@ private:
 	// Callback for getting the enabled state of the 'Reset to Defaults' button.
 	bool HandleResetToDefaultsButtonEnabled( ) const;
 
+	// Callback for getting the visibility of options that are only visible when editing something that came from a non-default object
+	EVisibility EditingNonDefaultSettingsVisibility( ) const;
+
 	// Callback for navigating a settings section link.
 	void HandleSectionLinkNavigate( ISettingsSectionPtr Section );
 
@@ -169,6 +172,12 @@ private:
 
 	// Reports a preference changed event to the analytics system */
 	void RecordPreferenceChangedAnalytics( ISettingsSectionPtr SelectedSection, const FPropertyChangedEvent& PropertyChangedEvent ) const;
+
+	// Returns the cached config file name
+	FString GetCachedConfigFileName( ) const
+	{
+		return CachedConfigFileName;
+	}
 private:
 
 	// Holds the vertical box for settings categories.
diff --git a/Engine/Source/Developer/SettingsEditor/Private/Widgets/SSettingsEditorCheckoutNotice.cpp b/Engine/Source/Developer/SettingsEditor/Private/Widgets/SSettingsEditorCheckoutNotice.cpp
--- a/Engine/Source/Developer/SettingsEditor/Private/Widgets/SSettingsEditorCheckoutNotice.cpp
+++ b/Engine/Source/Developer/SettingsEditor/Private/Widgets/SSettingsEditorCheckoutNotice.cpp
@@ -13,14 +13,15 @@ void SSettingsEditorCheckoutNotice::Construct( const FArguments& InArgs )
 {
 	CheckOutClickedDelegate = InArgs._OnCheckOutClicked;
 	bIsUnlocked = InArgs._Unlocked;
+	ConfigFilePath = InArgs._ConfigFilePath;
 	bLookingForSourceControlState = InArgs._LookingForSourceControlState;
 
 	// default configuration notice
 	ChildSlot
 	[
 		SNew(SBorder)
-			.BorderBackgroundColor(FLinearColor::Yellow)
-			.BorderImage(FEditorStyle::GetBrush("ToolPanel.GroupBorder"))
+			.BorderBackgroundColor(this, &SSettingsEditorCheckoutNotice::GetLockedOrUnlockedStatusBarColor)
+			.BorderImage(FEditorStyle::GetBrush("ToolPanel.LightGroupBorder"))
 			.Padding(8.0f)
 			[
 				SNew(SWidgetSwitcher)
@@ -48,6 +49,9 @@ void SSettingsEditorCheckoutNotice::Construct( const FArguments& InArgs )
 							[
 								SNew(STextBlock)
 									.Text(this, &SSettingsEditorCheckoutNotice::HandleLockedStatusText)
+									.ColorAndOpacity(FLinearColor::White)
+									.ShadowColorAndOpacity(FLinearColor::Black)
+									.ShadowOffset(FVector2D::UnitVector)
 							]
 
 						// Check out button
@@ -93,6 +97,9 @@ void SSettingsEditorCheckoutNotice::Construct( const FArguments& InArgs )
 							[
 								SNew(STextBlock)
 									.Text(this, &SSettingsEditorCheckoutNotice::HandleUnlockedStatusText)
+									.ColorAndOpacity(FLinearColor::White)
+									.ShadowColorAndOpacity(FLinearColor::Black)
+									.ShadowOffset(FVector2D::UnitVector)
 							]
 					]
 			]
@@ -148,18 +155,28 @@ EVisibility SSettingsEditorCheckoutNotice::HandleCheckOutButtonVisibility( ) con
 
 FText SSettingsEditorCheckoutNotice::HandleLockedStatusText() const
 {
-	return ISourceControlModule::Get().IsEnabled() ? 
-		LOCTEXT("DefaultSettingsNotice_WithSourceControl", "These settings are always saved in the default configuration file, which is currently not checked out.") :
-		LOCTEXT("DefaultSettingsNotice_Source", "These settings are always saved in the default configuration file, which is currently not writable.");
+	FText ConfigFilename = FText::FromString(FPaths::GetCleanFilename(ConfigFilePath.Get()));
+
+	return FText::Format(ISourceControlModule::Get().IsEnabled() ?
+		LOCTEXT("DefaultSettingsNotice_WithSourceControl", "These settings are saved in {0}, which is not currently checked out.") :
+		LOCTEXT("DefaultSettingsNotice_Source", "These settings are saved {0}, which is not currently writable."), ConfigFilename);
 }
 
 FText SSettingsEditorCheckoutNotice::HandleUnlockedStatusText() const
 {
-	return ISourceControlModule::Get().IsEnabled() ? 
-		LOCTEXT("DefaultSettingsNotice_CheckedOut", "These settings are always saved in the default configuration file, which is currently checked out.") :
-		LOCTEXT("DefaultSettingsNotice_Writable", "These settings are always saved in the default configuration file, which is currently writable.");
+	FText ConfigFilename = FText::FromString(FPaths::GetCleanFilename(ConfigFilePath.Get()));
+	
+	return FText::Format(ISourceControlModule::Get().IsEnabled() ?
+		LOCTEXT("DefaultSettingsNotice_CheckedOut", "These settings are saved in {0}, which is currently checked out.") :
+		LOCTEXT("DefaultSettingsNotice_Writable", "These settings are saved in {0}, which is currently writable."), ConfigFilename);
 }
 
+FSlateColor SSettingsEditorCheckoutNotice::GetLockedOrUnlockedStatusBarColor() const
+{
+	return bIsUnlocked.Get() ? FLinearColor::Green : FLinearColor::Yellow;
+}
+
+
 EVisibility SSettingsEditorCheckoutNotice::HandleThrobberVisibility() const
 {
 	if(ISourceControlModule::Get().IsEnabled())
diff --git a/Engine/Source/Developer/SettingsEditor/Private/Widgets/SSettingsEditorCheckoutNotice.h b/Engine/Source/Developer/SettingsEditor/Private/Widgets/SSettingsEditorCheckoutNotice.h
--- a/Engine/Source/Developer/SettingsEditor/Private/Widgets/SSettingsEditorCheckoutNotice.h
+++ b/Engine/Source/Developer/SettingsEditor/Private/Widgets/SSettingsEditorCheckoutNotice.h
@@ -15,7 +15,10 @@ public:
 
 		/** Called to determine if the associated file is unlocked */
 		SLATE_ATTRIBUTE(bool, Unlocked)
-	
+
+		/** Called to get the filename of the config file for display */
+		SLATE_ATTRIBUTE(FString, ConfigFilePath)
+
 		/** Slot for this button's content (optional) */
 		SLATE_NAMED_SLOT(FArguments, LockedContent)
 
@@ -60,6 +63,9 @@ private:
 	// Callback for getting the status text when the config is unlocked
 	FText HandleUnlockedStatusText() const;
 
+	// Callback for getting the status color, which indicates if it is writeable or not
+	FSlateColor GetLockedOrUnlockedStatusBarColor() const;
+
 	// Callback for getting the visibility of the source control throbber
 	EVisibility HandleThrobberVisibility() const;
 
@@ -70,5 +76,7 @@ private:
 
 	TAttribute<bool> bIsUnlocked;
 
+	TAttribute<FString> ConfigFilePath;
+
 	TAttribute<bool> bLookingForSourceControlState;
 };
diff --git a/Engine/Source/Developer/SharedSettingsWidgets/Private/SPlatformSetupMessage.cpp b/Engine/Source/Developer/SharedSettingsWidgets/Private/SPlatformSetupMessage.cpp
--- a/Engine/Source/Developer/SharedSettingsWidgets/Private/SPlatformSetupMessage.cpp
+++ b/Engine/Source/Developer/SharedSettingsWidgets/Private/SPlatformSetupMessage.cpp
@@ -33,6 +33,9 @@ TSharedRef<SWidget> SPlatformSetupMessage::MakeRow(FName IconName, FText Message
 		.VAlign(VAlign_Center)
 		[
 			SNew(STextBlock)
+			.ColorAndOpacity(FLinearColor::White)
+			.ShadowColorAndOpacity(FLinearColor::Black)
+			.ShadowOffset(FVector2D::UnitVector)
 			.Text(Message)
 		];
 
@@ -82,7 +85,7 @@ void SPlatformSetupMessage::Construct(const FArguments& InArgs, const FString& I
 	[
 		SNew(SBorder)
 		.BorderBackgroundColor(this, &SPlatformSetupMessage::GetBorderColor)
-		.BorderImage(FEditorStyle::GetBrush("ToolPanel.GroupBorder"))
+		.BorderImage(FEditorStyle::GetBrush("ToolPanel.LightGroupBorder"))
 		.Padding(8.0f)
 		[
 			SNew(SWidgetSwitcher)
diff --git a/Engine/Source/Editor/EditorStyle/Private/SlateEditorStyle.cpp b/Engine/Source/Editor/EditorStyle/Private/SlateEditorStyle.cpp
--- a/Engine/Source/Editor/EditorStyle/Private/SlateEditorStyle.cpp
+++ b/Engine/Source/Editor/EditorStyle/Private/SlateEditorStyle.cpp
@@ -510,6 +510,7 @@ void FSlateEditorStyle::FStyle::SetupGeneralStyles()
 	{
 		Set( "ToolPanel.GroupBorder", new BOX_BRUSH( "Common/GroupBorder", FMargin(4.0f/16.0f) ) );
 		Set( "ToolPanel.DarkGroupBorder", new BOX_BRUSH( "Common/DarkGroupBorder", FMargin( 4.0f / 16.0f ) ) );
+		Set( "ToolPanel.LightGroupBorder", new BOX_BRUSH("Common/LightGroupBorder", FMargin(4.0f / 16.0f)) );
 	}
 
 	// Inline Editable Text Block
