Editor: Update hardware targetting to use SSettingsEditorCheckoutNotice and devolved much of the bookkeeping in SSettingsEditor into SSettingsEditorCheckoutNotice

diff --git a/Engine/Source/Developer/SettingsEditor/Private/Widgets/SSettingsEditor.cpp b/Engine/Source/Developer/SettingsEditor/Private/Widgets/SSettingsEditor.cpp
--- a/Engine/Source/Developer/SettingsEditor/Private/Widgets/SSettingsEditor.cpp
+++ b/Engine/Source/Developer/SettingsEditor/Private/Widgets/SSettingsEditor.cpp
@@ -26,9 +26,6 @@ SSettingsEditor::~SSettingsEditor( )
 void SSettingsEditor::Construct( const FArguments& InArgs, const ISettingsEditorModelRef& InModel )
 {
 	Model = InModel;
-	DefaultConfigCheckOutTimer = 0.0f;
-	DefaultConfigCheckOutNeeded = false;
-	DefaultConfigQueryInProgress = false;
 	SettingsContainer = InModel->GetSettingsContainer();
 
 	// initialize settings view
@@ -49,6 +46,12 @@ void SSettingsEditor::Construct( const FArguments& InArgs, const ISettingsEditor
 	SettingsView->SetVisibility(TAttribute<EVisibility>::Create(TAttribute<EVisibility>::FGetter::CreateSP(this, &SSettingsEditor::HandleSettingsViewVisibility)));
 	SettingsView->SetIsPropertyEditingEnabledDelegate(FIsPropertyEditingEnabled::CreateSP(this, &SSettingsEditor::HandleSettingsViewEnabled));
 
+	// Create the watcher widget for the default config file (checks file status / SCC state)
+	FileWatcherWidget =
+		SNew(SSettingsEditorCheckoutNotice)
+		.Visibility(this, &SSettingsEditor::HandleDefaultConfigNoticeVisibility)
+		.ConfigFilePath(this, &SSettingsEditor::GetConfigFileName);
+
 	ChildSlot
 	[
 		SNew(SVerticalBox)
@@ -194,12 +197,7 @@ void SSettingsEditor::Construct( const FArguments& InArgs, const ISettingsEditor
 							.AutoHeight()
 							.Padding(0.0f, 16.0f, 0.0f, 16.0f)
 							[
-								// checkout notice
-								SNew(SSettingsEditorCheckoutNotice)
-									.Visibility(this, &SSettingsEditor::HandleDefaultConfigNoticeVisibility)
-									.Unlocked(this, &SSettingsEditor::HandleConfigNoticeUnlocked)
-									.ConfigFilePath(this, &SSettingsEditor::GetCachedConfigFileName)
-									.LookingForSourceControlState(this, &SSettingsEditor::HandleLookingForSourceControlState)
+								FileWatcherWidget.ToSharedRef()
 							]
 
 						+ SVerticalBox::Slot()
@@ -228,58 +226,6 @@ void SSettingsEditor::Construct( const FArguments& InArgs, const ISettingsEditor
 	ReloadCategories();
 }
 
-
-/* SCompoundWidget interface
- *****************************************************************************/
-
-void SSettingsEditor::Tick( const FGeometry& AllottedGeometry, const double InCurrentTime, const float InDeltaTime )
-{
-	SCompoundWidget::Tick(AllottedGeometry, InCurrentTime, InDeltaTime);
-
-	// cache selected settings object's configuration file state
-	DefaultConfigCheckOutTimer += InDeltaTime;
-
-	if (DefaultConfigCheckOutTimer >= 1.0f)
-	{
-		TWeakObjectPtr<UObject> SettingsObject = GetSelectedSettingsObject();
-
-		if (SettingsObject.IsValid() && SettingsObject->GetClass()->HasAnyClassFlags(CLASS_Config | CLASS_DefaultConfig))
-		{
-			CachedConfigFileName = GetDefaultConfigFilePath(SettingsObject);
-			bool NewCheckOutNeeded = false;
-			if(ISourceControlModule::Get().IsEnabled())
-			{
-				// note: calling QueueStatusUpdate often does not spam status updates as an internal timer prevents this
-				ISourceControlModule::Get().QueueStatusUpdate(CachedConfigFileName);
-
-				ISourceControlProvider& SourceControlProvider = ISourceControlModule::Get().GetProvider();
-				FSourceControlStatePtr SourceControlState = SourceControlProvider.GetState(CachedConfigFileName, EStateCacheUsage::Use);
-				NewCheckOutNeeded = SourceControlState.IsValid() && SourceControlState->CanCheckout();
-				DefaultConfigQueryInProgress = SourceControlState.IsValid() && SourceControlState->IsUnknown();
-			}
-			else
-			{
-				NewCheckOutNeeded = (FPaths::FileExists(CachedConfigFileName) && IFileManager::Get().IsReadOnly(*CachedConfigFileName));
-			}
-
-			// file has been checked in or reverted
-			if ((NewCheckOutNeeded == true) && (DefaultConfigCheckOutNeeded == false))
-			{
-				SettingsObject->ReloadConfig();
-			}
-
-			DefaultConfigCheckOutNeeded = NewCheckOutNeeded;
-		}
-		else
-		{
-			DefaultConfigCheckOutNeeded = false;
-		}
-
-		DefaultConfigCheckOutTimer = 0.0f;
-	}
-}
-
-
 /* FNotifyHook interface
  *****************************************************************************/
 
@@ -516,20 +462,6 @@ void SSettingsEditor::ShowNotification( const FText& Text, SNotificationItem::EC
 /* SSettingsEditor callbacks
  *****************************************************************************/
 
-FReply SSettingsEditor::HandleCheckOutButtonClicked( )
-{
-	if (ISourceControlModule::Get().IsEnabled())
-	{
-		CheckOutDefaultConfigFile();
-	}
-	else
-	{
-		MakeDefaultConfigFileWritable();
-	}
-
-	return FReply::Handled();
-}
-
 
 void SSettingsEditor::HandleCultureChanged( )
 {
@@ -537,15 +469,6 @@ void SSettingsEditor::HandleCultureChanged( )
 }
 
 
-bool SSettingsEditor::HandleConfigNoticeUnlocked( ) const
-{
-	return !DefaultConfigCheckOutNeeded && !DefaultConfigQueryInProgress;
-}
-
-bool SSettingsEditor::HandleLookingForSourceControlState() const
-{
-	return DefaultConfigQueryInProgress;
-}
 
 void SSettingsEditor::RecordPreferenceChangedAnalytics( ISettingsSectionPtr SelectedSection, const FPropertyChangedEvent& PropertyChangedEvent ) const
 {
@@ -562,6 +485,20 @@ void SSettingsEditor::RecordPreferenceChangedAnalytics( ISettingsSectionPtr Sele
 	}
 }
 
+FString SSettingsEditor::GetConfigFileName() const
+{
+	TWeakObjectPtr<UObject> SettingsObject = GetSelectedSettingsObject();
+
+	if (SettingsObject.IsValid() && SettingsObject->GetClass()->HasAnyClassFlags(CLASS_Config | CLASS_DefaultConfig))
+	{
+		return GetDefaultConfigFilePath(SettingsObject);
+	}
+	else
+	{
+		return FString();
+	}
+}
+
 EVisibility SSettingsEditor::HandleDefaultConfigNoticeVisibility( ) const
 {
 	TWeakObjectPtr<UObject> SettingsObject = GetSelectedSettingsObject();
@@ -656,7 +593,7 @@ bool SSettingsEditor::HandleImportButtonEnabled( ) const
 
 	if (SelectedSection.IsValid())
 	{
-		return (SelectedSection->CanEdit() && SelectedSection->CanImport() && !DefaultConfigCheckOutNeeded);
+		return SelectedSection->CanEdit() && SelectedSection->CanImport() && !IsDefaultConfigCheckOutNeeded();
 	}
 
 	return false;
@@ -708,7 +645,7 @@ void SSettingsEditor::HandleModelSelectionChanged( )
 		SettingsView->SetObject(nullptr);
 	}
 
-	DefaultConfigCheckOutTimer = 1.0f;
+	FileWatcherWidget->Invalidate();
 }
 
 
@@ -769,7 +706,7 @@ FReply SSettingsEditor::HandleSetAsDefaultButtonClicked( )
 
 	if (SelectedSection.IsValid())
 	{
-		if (DefaultConfigCheckOutNeeded)
+		if (IsDefaultConfigCheckOutNeeded())
 		{
 			if (ISourceControlModule::Get().IsEnabled())
 			{
@@ -862,7 +799,7 @@ bool SSettingsEditor::HandleSettingsViewEnabled( ) const
 {
 	ISettingsSectionPtr SelectedSection = Model->GetSelectedSection();
 
-	return (SelectedSection.IsValid() && SelectedSection->CanEdit() && (!SelectedSection->HasDefaultSettingsObject() || (!DefaultConfigCheckOutNeeded && !DefaultConfigQueryInProgress)));
+	return (SelectedSection.IsValid() && SelectedSection->CanEdit() && (!SelectedSection->HasDefaultSettingsObject() || FileWatcherWidget->IsUnlocked()));
 }
 
 
@@ -879,4 +816,18 @@ EVisibility SSettingsEditor::HandleSettingsViewVisibility( ) const
 }
 
 
+// Do we need to edit the default config file?
+bool SSettingsEditor::IsDefaultConfigCheckOutNeeded() const
+{
+	TWeakObjectPtr<UObject> SettingsObject = GetSelectedSettingsObject();
+	if (SettingsObject.IsValid() && SettingsObject->GetClass()->HasAnyClassFlags(CLASS_Config | CLASS_DefaultConfig))
+	{
+		return !FileWatcherWidget->IsUnlocked();
+	}
+	else
+	{
+		return false;
+	}
+}
+
 #undef LOCTEXT_NAMESPACE
diff --git a/Engine/Source/Developer/SettingsEditor/Private/Widgets/SSettingsEditor.h b/Engine/Source/Developer/SettingsEditor/Private/Widgets/SSettingsEditor.h
--- a/Engine/Source/Developer/SettingsEditor/Private/Widgets/SSettingsEditor.h
+++ b/Engine/Source/Developer/SettingsEditor/Private/Widgets/SSettingsEditor.h
@@ -32,12 +32,6 @@ public:
 	 */
 	void Construct( const FArguments& InArgs, const ISettingsEditorModelRef& InModel );
 
-public:
-
-	// SCompoundWidget interface
-
-	virtual void Tick( const FGeometry& AllottedGeometry, const double InCurrentTime, const float InDeltaTime ) override;
-
 public:
 
 	// FNotifyHook interface
@@ -101,12 +95,6 @@ private:
 	// Callback for clicking the Back link.
 	FReply HandleBackButtonClicked( );
 
-	// Callback for clicking the 'Check Out' button.
-	FReply HandleCheckOutButtonClicked( );
-
-	// Callback for getting the widget index for the notice switcher.
-	bool HandleConfigNoticeUnlocked( ) const;
-
 	// Callback for when the user's culture has changed.
 	void HandleCultureChanged( );
 
@@ -167,31 +155,22 @@ private:
 	// Callback for determining the visibility of the settings view.
 	EVisibility HandleSettingsViewVisibility( ) const;
 
-	// Callback for determining whether we are currently looking for the source control state of the config file
-	bool HandleLookingForSourceControlState( ) const;
-
 	// Reports a preference changed event to the analytics system */
 	void RecordPreferenceChangedAnalytics( ISettingsSectionPtr SelectedSection, const FPropertyChangedEvent& PropertyChangedEvent ) const;
 
-	// Returns the cached config file name
-	FString GetCachedConfigFileName( ) const
-	{
-		return CachedConfigFileName;
-	}
+	// Returns the config file name currently being edited
+	FString GetConfigFileName() const;
+
+	// Do we need to edit the default config file?
+	bool IsDefaultConfigCheckOutNeeded() const;
 private:
 
+	// Watcher widget for the default config file (checks file status / SCC state)
+	TSharedPtr<class SSettingsEditorCheckoutNotice> FileWatcherWidget;
+
 	// Holds the vertical box for settings categories.
 	TSharedPtr<SVerticalBox> CategoriesBox;
 
-	// Holds a flag indicating whether the section's configuration file needs to be checked out.
-	bool DefaultConfigCheckOutNeeded;
-
-	// Holds a flag indicating whether the section's configuration file has a source control status request in progress.
-	bool DefaultConfigQueryInProgress;
-
-	// Holds a timer for checking whether the section's configuration file needs to be checked out.
-	float DefaultConfigCheckOutTimer;
-
 	// Holds the overlay slot for custom widgets.
 	SOverlay::FOverlaySlot* CustomWidgetSlot;
 
@@ -206,7 +185,4 @@ private:
 
 	// Holds the details view.
 	TSharedPtr<IDetailsView> SettingsView;
-
-	// Cached filename for the config file we are displaying
-	FString CachedConfigFileName;
 };
diff --git a/Engine/Source/Developer/SharedSettingsWidgets/Private/SSettingsEditorCheckoutNotice.cpp b/Engine/Source/Developer/SharedSettingsWidgets/Private/SSettingsEditorCheckoutNotice.cpp
--- a/Engine/Source/Developer/SharedSettingsWidgets/Private/SSettingsEditorCheckoutNotice.cpp
+++ b/Engine/Source/Developer/SharedSettingsWidgets/Private/SSettingsEditorCheckoutNotice.cpp
@@ -15,9 +15,11 @@
 void SSettingsEditorCheckoutNotice::Construct( const FArguments& InArgs )
 {
 	CheckOutClickedDelegate = InArgs._OnCheckOutClicked;
-	bIsUnlocked = InArgs._Unlocked;
 	ConfigFilePath = InArgs._ConfigFilePath;
-	bLookingForSourceControlState = InArgs._LookingForSourceControlState;
+
+	DefaultConfigCheckOutTimer = 0.0f;
+	DefaultConfigCheckOutNeeded = false;
+	DefaultConfigQueryInProgress = false;
 
 	// default configuration notice
 	ChildSlot
@@ -171,12 +173,11 @@ FText SSettingsEditorCheckoutNotice::HandleCheckOutButtonToolTip( ) const
 	return LOCTEXT("MakeWritableTooltip", "Make the default configuration file that holds these settings writable.");
 }
 
-
 EVisibility SSettingsEditorCheckoutNotice::HandleCheckOutButtonVisibility( ) const
 {
 	if (ISourceControlModule::Get().IsEnabled() && ISourceControlModule::Get().GetProvider().IsAvailable())
 	{
-		return !bLookingForSourceControlState.Get() ? EVisibility::Visible : EVisibility::Collapsed;
+		return !DefaultConfigQueryInProgress ? EVisibility::Visible : EVisibility::Hidden;
 	}
 
 	return EVisibility::Collapsed;
@@ -202,7 +203,7 @@ FText SSettingsEditorCheckoutNotice::HandleUnlockedStatusText() const
 
 FSlateColor SSettingsEditorCheckoutNotice::GetLockedOrUnlockedStatusBarColor() const
 {
-	return bIsUnlocked.Get() ? FLinearColor::Green : FLinearColor::Yellow;
+	return IsUnlocked() ? FLinearColor::Green : FLinearColor::Yellow;
 }
 
 
@@ -210,9 +211,64 @@ EVisibility SSettingsEditorCheckoutNotice::HandleThrobberVisibility() const
 {
 	if(ISourceControlModule::Get().IsEnabled())
 	{
-		return bLookingForSourceControlState.Get() ? EVisibility::Visible : EVisibility::Collapsed;
+		return DefaultConfigQueryInProgress ? EVisibility::Visible : EVisibility::Collapsed;
 	}
 	return EVisibility::Collapsed;
 }
 
+
+void SSettingsEditorCheckoutNotice::Tick(const FGeometry& AllottedGeometry, const double InCurrentTime, const float InDeltaTime)
+{
+	SCompoundWidget::Tick(AllottedGeometry, InCurrentTime, InDeltaTime);
+
+	// cache selected settings object's configuration file state
+	DefaultConfigCheckOutTimer += InDeltaTime;
+
+	if (DefaultConfigCheckOutTimer >= 1.0f)
+	{
+		bool NewCheckOutNeeded = false;
+
+		DefaultConfigQueryInProgress = false;
+		FString CachedConfigFileName = ConfigFilePath.Get();
+		if (!CachedConfigFileName.IsEmpty())
+		{
+			if (ISourceControlModule::Get().IsEnabled())
+			{
+				// note: calling QueueStatusUpdate often does not spam status updates as an internal timer prevents this
+				ISourceControlModule::Get().QueueStatusUpdate(CachedConfigFileName);
+
+				ISourceControlProvider& SourceControlProvider = ISourceControlModule::Get().GetProvider();
+				FSourceControlStatePtr SourceControlState = SourceControlProvider.GetState(CachedConfigFileName, EStateCacheUsage::Use);
+				NewCheckOutNeeded = SourceControlState.IsValid() && SourceControlState->CanCheckout();
+				DefaultConfigQueryInProgress = SourceControlState.IsValid() && SourceControlState->IsUnknown();
+			}
+			else
+			{
+				NewCheckOutNeeded = (FPaths::FileExists(CachedConfigFileName) && IFileManager::Get().IsReadOnly(*CachedConfigFileName));
+			}
+
+			// file has been checked in or reverted
+			if ((NewCheckOutNeeded == true) && (DefaultConfigCheckOutNeeded == false))
+			{
+				// 				TWeakObjectPtr<UObject> SettingsObject = GetSelectedSettingsObject();
+				// 
+				// 				if (SettingsObject.IsValid() && SettingsObject->GetClass()->HasAnyClassFlags(CLASS_Config | CLASS_DefaultConfig))
+				// 				{
+				// 					//@TODO: This thing!
+				// 					SettingsObject->ReloadConfig();
+				// 				}
+			}
+		}
+
+		DefaultConfigCheckOutNeeded = NewCheckOutNeeded;
+		DefaultConfigCheckOutTimer = 0.0f;
+	}
+}
+
+
+bool SSettingsEditorCheckoutNotice::IsUnlocked() const
+{
+	return !DefaultConfigCheckOutNeeded && !DefaultConfigQueryInProgress;
+}
+
 #undef LOCTEXT_NAMESPACE
diff --git a/Engine/Source/Developer/SharedSettingsWidgets/Public/SSettingsEditorCheckoutNotice.h b/Engine/Source/Developer/SharedSettingsWidgets/Public/SSettingsEditorCheckoutNotice.h
--- a/Engine/Source/Developer/SharedSettingsWidgets/Public/SSettingsEditorCheckoutNotice.h
+++ b/Engine/Source/Developer/SharedSettingsWidgets/Public/SSettingsEditorCheckoutNotice.h
@@ -36,6 +36,19 @@ public:
 	 */
 	void Construct( const FArguments& InArgs );
 
+	// SWidget interface
+	virtual void Tick(const FGeometry& AllottedGeometry, const double InCurrentTime, const float InDeltaTime) override;
+	// End of SWidget interface
+
+	// Return true if the watched file can be modified (as of the last update)
+	bool IsUnlocked() const;
+
+	// Force an update on the next tick
+	void Invalidate()
+	{
+		DefaultConfigCheckOutTimer = 1.0f;
+	}
+
 private:
 
 	// Callback for clicking the 'Check Out' button.
@@ -53,7 +66,7 @@ private:
 	// Callback for getting the widget index for the notice switcher.
 	int32 HandleNoticeSwitcherWidgetIndex( ) const
 	{
-		return bIsUnlocked.Get() ? 1 : 0;
+		return IsUnlocked() ? 1 : 0;
 	}
 
 	// Callback for getting the status text when the config is locked
@@ -73,9 +86,15 @@ private:
 	// Holds a delegate that is executed when the 'Check Out' button has been clicked.
 	FOnClicked CheckOutClickedDelegate;
 
-	TAttribute<bool> bIsUnlocked;
-
+	// The current file being watched
 	TAttribute<FString> ConfigFilePath;
 
-	TAttribute<bool> bLookingForSourceControlState;
+	// Holds a flag indicating whether the section's configuration file needs to be checked out.
+	bool DefaultConfigCheckOutNeeded;
+
+	// Holds a flag indicating whether the section's configuration file has a source control status request in progress.
+	bool DefaultConfigQueryInProgress;
+
+	// Holds a timer for checking whether the section's configuration file needs to be checked out.
+	float DefaultConfigCheckOutTimer;
 };
diff --git a/Engine/Source/Editor/DetailCustomizations/DetailCustomizations.Build.cs b/Engine/Source/Editor/DetailCustomizations/DetailCustomizations.Build.cs
--- a/Engine/Source/Editor/DetailCustomizations/DetailCustomizations.Build.cs
+++ b/Engine/Source/Editor/DetailCustomizations/DetailCustomizations.Build.cs
@@ -41,6 +41,7 @@ public class DetailCustomizations : ModuleRules
                 "Media",
                 "MediaAssets",
                 "HardwareTargeting",
+				"SharedSettingsWidgets"
 			}
 		);
 
diff --git a/Engine/Source/Editor/DetailCustomizations/Private/HardwareTargetingSettingsDetails.cpp b/Engine/Source/Editor/DetailCustomizations/Private/HardwareTargetingSettingsDetails.cpp
--- a/Engine/Source/Editor/DetailCustomizations/Private/HardwareTargetingSettingsDetails.cpp
+++ b/Engine/Source/Editor/DetailCustomizations/Private/HardwareTargetingSettingsDetails.cpp
@@ -5,6 +5,7 @@
 #include "HardwareTargetingSettingsDetails.h"
 #include "HardwareTargetingModule.h"
 #include "ISourceControlModule.h"
+#include "SSettingsEditorCheckoutNotice.h"
 
 #define LOCTEXT_NAMESPACE "FHardwareTargetingSettingsDetails"
 
@@ -123,9 +124,9 @@ public:
 
 	bool CanApply() const
 	{
-		for (const auto& Pair : WritableStatus)
+		for (const auto& Watcher : FileWatcherWidgets)
 		{
-			if (Pair.Value != EWritableStatus::Writable)
+			if (!Watcher->IsUnlocked())
 			{
 				return false;
 			}
@@ -139,72 +140,9 @@ public:
 		return FPaths::ConvertRelativePathToFull(FString::Printf(TEXT("%sDefault%s.ini"), *FPaths::SourceConfigDir(), *Settings->GetClass()->ClassConfigName.ToString()));
 	}
 
-	virtual void Tick(const FGeometry& AllottedGeometry, const double InCurrentTime, const float InDeltaTime)
-	{
-		if (!ISourceControlModule::Get().IsEnabled())
-		{
-			return;
-		}
-
-		const double CheckStatusIntervalS = 2;
-		if (InCurrentTime - LastStatusUpdate > CheckStatusIntervalS)
-		{
-			LastStatusUpdate = InCurrentTime;
-
-			for (auto& Pair : WritableStatus)
-			{
-				if (Pair.Key.IsValid() && Pair.Value == EWritableStatus::Unknown)
-				{
-					ISourceControlProvider& SourceControlProvider = ISourceControlModule::Get().GetProvider();
-					FSourceControlStatePtr SourceControlState = SourceControlProvider.GetState(GetDefaultConfigPath(Pair.Key.Get()), EStateCacheUsage::Use);
-					if (SourceControlState.IsValid())
-					{
-						Pair.Value = SourceControlState->IsCheckedOut() ? EWritableStatus::Writable : EWritableStatus::NotWritable;
-					}
-				}
-			}
-		}
-	}
-
-	bool AreSettingsWritable(TWeakObjectPtr<UObject> Settings) const
-	{
-		if (auto* Status = WritableStatus.Find(Settings))
-		{
-			if (*Status == EWritableStatus::Writable)
-			{
-				return true;
-			}
-		}
-		return false;
-	}
-
-	EVisibility GetWritableHintVisibility(TWeakObjectPtr<UObject> Settings) const
-	{
-		if (const auto* Status = WritableStatus.Find(Settings))
-		{
-			if (*Status == EWritableStatus::NotWritable)
-			{
-				return EVisibility::Visible;
-			}
-		}
-		return EVisibility::Collapsed;
-	}
-
-	EVisibility GetWritableThrobberVisibility(TWeakObjectPtr<UObject> Settings) const
-	{
-		if (const auto* Status = WritableStatus.Find(Settings))
-		{
-			if (*Status == EWritableStatus::Unknown)
-			{
-				return EVisibility::Visible;
-			}
-		}
-		return EVisibility::Collapsed;
-	}
-
 	void Update()
 	{
-		WritableStatus.Empty();
+		FileWatcherWidgets.Reset();
 		SettingsDescriptions->ClearChildren();
 
 		IHardwareTargetingModule& Module = IHardwareTargetingModule::Get();
@@ -226,85 +164,28 @@ public:
 				continue;
 			}
 			
-			EWritableStatus& Status = WritableStatus.Add(Settings.SettingsObject, EWritableStatus::Unknown);
+
 			FString ConfigFile = GetDefaultConfigPath(Settings.SettingsObject.Get());
 
 			if (!SeenConfigFiles.Contains(ConfigFile))
 			{
 				SeenConfigFiles.Add(ConfigFile);
 
-				if (ISourceControlModule::Get().IsEnabled())
-				{
-					ISourceControlModule::Get().QueueStatusUpdate(ConfigFile);
-
-					ISourceControlProvider& SourceControlProvider = ISourceControlModule::Get().GetProvider();
-					FSourceControlStatePtr SourceControlState = SourceControlProvider.GetState(ConfigFile, EStateCacheUsage::Use);
-					if (SourceControlState.IsValid())
-					{
-						Status = SourceControlState->IsCheckedOut() ? EWritableStatus::Writable : EWritableStatus::NotWritable;
-					}
-				}
-				else
-				{
-					const bool bIsWritable = !FPaths::FileExists(ConfigFile) || !IFileManager::Get().IsReadOnly(*ConfigFile);
-					Status = bIsWritable ? EWritableStatus::Writable : EWritableStatus::NotWritable;
-				}
+				TSharedRef<SSettingsEditorCheckoutNotice> FileWatcherWidget =
+					SNew(SSettingsEditorCheckoutNotice)
+					.ConfigFilePath(ConfigFile);
+				FileWatcherWidgets.Add(FileWatcherWidget);
 
 				SettingsDescriptions->AddSlot(0, SlotIndex)
 				[
-					SNew(SVerticalBox)
-
-					+ SVerticalBox::Slot()
-					.AutoHeight()
-					[
-						SNew(STextBlock)
-						.Text(FPaths::GetCleanFilename(ConfigFile))
-					]
-				 
-					+ SVerticalBox::Slot()
-					.AutoHeight()
-					[
-						SNew(STextBlock)
-						.Visibility(this, &SRequiredDefaultConfig::GetWritableHintVisibility, Settings.SettingsObject)
-						.Text(LOCTEXT("NotWritable", "This file is not writable"))
-					]
-				 
-					+ SVerticalBox::Slot()
-					.AutoHeight()
-					[
-						SNew(SBorder)
-						.Visibility(this, &SRequiredDefaultConfig::GetWritableThrobberVisibility, Settings.SettingsObject)
-						.BorderImage(FEditorStyle::GetBrush("NoBorder"))
-						.Padding(0)
-						[
-							SNew(SHorizontalBox)
-
-							+ SHorizontalBox::Slot()
-							.AutoWidth()
-							[
-								SNew(SThrobber)
-
-							]
-							+ SHorizontalBox::Slot()
-							.AutoWidth()
-							[
-								SNew(STextBlock)
-								.Text(LOCTEXT("CheckingFileStatus", "Checking file status"))
-							]
-						]
-					]
+					FileWatcherWidget
 				];
+				SlotIndex++;
 			}
-
-			++SlotIndex;
 		}
 	}
-	
-	enum class EWritableStatus
-	{
-		Unknown, Writable, NotWritable
-	};
-	TMap<TWeakObjectPtr<UObject>, EWritableStatus> WritableStatus;
+
+	TArray<TSharedPtr<SSettingsEditorCheckoutNotice>> FileWatcherWidgets;
 
 	TSharedPtr<SUniformGridPanel> SettingsDescriptions;
 
