SettingsEditor: API improvements

- made public headers compilable individually
- moved interfaces to public root
- removed module singleton getter
- code and documentation cleanup

#UpgradeNotes:
- instead of ISettingsEditorModule::GetRef() use FModuleManager::GetModuleChecked<ISettingsModule>("SettingsModule")
- instead of including SettingsEditor.h include the interface headers you actually use

diff --git a/Engine/Plugins/Editor/MacGraphicsSwitching/Source/MacGraphicsSwitching/Private/MacGraphicsSwitchingModule.cpp b/Engine/Plugins/Editor/MacGraphicsSwitching/Source/MacGraphicsSwitching/Private/MacGraphicsSwitchingModule.cpp
--- a/Engine/Plugins/Editor/MacGraphicsSwitching/Source/MacGraphicsSwitching/Private/MacGraphicsSwitchingModule.cpp
+++ b/Engine/Plugins/Editor/MacGraphicsSwitching/Source/MacGraphicsSwitching/Private/MacGraphicsSwitchingModule.cpp
@@ -8,7 +8,6 @@
 #include "MacGraphicsSwitchingWidget.h"
 #include "Runtime/Core/Public/Features/IModularFeatures.h"
 #include "Runtime/SlateCore/Public/Rendering/SlateRenderer.h"
-#include "Developer/Settings/Public/Settings.h"
 #include "Editor/MainFrame/Public/Interfaces/IMainFrameModule.h"
 #include "ISettingsModule.h"
 #include "LevelEditor.h"
diff --git a/Engine/Source/Developer/SettingsEditor/Private/Models/SettingsEditorModel.h b/Engine/Source/Developer/SettingsEditor/Private/Models/SettingsEditorModel.h
--- a/Engine/Source/Developer/SettingsEditor/Private/Models/SettingsEditorModel.h
+++ b/Engine/Source/Developer/SettingsEditor/Private/Models/SettingsEditorModel.h
@@ -22,10 +22,8 @@ public:
 		SettingsContainer->OnSectionRemoved().AddRaw(this, &FSettingsEditorModel::HandleSettingsContainerSectionRemoved);
 	}
 
-	/**
-	 * Destructor.
-	 */
-	~FSettingsEditorModel( )
+	/** Virtual destructor. */
+	virtual ~FSettingsEditorModel()
 	{
 		SettingsContainer->OnSectionRemoved().RemoveAll(this);
 	}
@@ -34,17 +32,17 @@ public:
 
 	// ISettingsEditorModel interface
 
-	virtual const ISettingsSectionPtr& GetSelectedSection( ) const override
+	virtual const ISettingsSectionPtr& GetSelectedSection() const override
 	{
 		return SelectedSection;
 	}
 
-	virtual const ISettingsContainerRef& GetSettingsContainer( ) const override
+	virtual const ISettingsContainerRef& GetSettingsContainer() const override
 	{
 		return SettingsContainer;
 	}
 
-	virtual FSimpleMulticastDelegate& OnSelectionChanged( ) override
+	virtual FSimpleMulticastDelegate& OnSelectionChanged() override
 	{
 		return OnSelectionChangedDelegate;
 	}
@@ -57,13 +55,12 @@ public:
 		}
 
 		SelectedSection = Section;
-
 		OnSelectionChangedDelegate.Broadcast();
 	}
 
 private:
 
-	// Handles the removal of sections from the settings container.
+	/** Handles the removal of sections from the settings container. */
 	void HandleSettingsContainerSectionRemoved( const ISettingsSectionRef& Section )
 	{
 		if (SelectedSection == Section)
@@ -74,14 +71,14 @@ private:
 
 private:
 
-	// Holds the currently selected settings section.
+	/** Holds the currently selected settings section. */
 	ISettingsSectionPtr SelectedSection;
 
-	// Holds a reference to the settings container.
+	/** Holds a reference to the settings container. */
 	ISettingsContainerRef SettingsContainer;
 
 private:
 
-	// Holds a delegate that is executed when the selected settings section has changed.
+	/** Holds a delegate that is executed when the selected settings section has changed. */
 	FSimpleMulticastDelegate OnSelectionChangedDelegate;
 };
diff --git a/Engine/Source/Developer/SettingsEditor/Private/SettingsEditorModule.cpp b/Engine/Source/Developer/SettingsEditor/Private/SettingsEditorModule.cpp
--- a/Engine/Source/Developer/SettingsEditor/Private/SettingsEditorModule.cpp
+++ b/Engine/Source/Developer/SettingsEditor/Private/SettingsEditorModule.cpp
@@ -1,7 +1,7 @@
 // Copyright 1998-2014 Epic Games, Inc. All Rights Reserved.
 
 #include "SettingsEditorPrivatePCH.h"
-#include "SSettingsEditor.h"
+#include "ISettingsEditorModule.h"
 
 
 /**
@@ -14,12 +14,12 @@ public:
 
 	// ISettingsEditorModule interface
 
-	virtual TSharedRef<SWidget> CreateEditor( const ISettingsEditorModelRef& Model ) override
+	virtual TSharedRef<SWidget> CreateEditor( const TSharedRef<ISettingsEditorModel>& Model ) override
 	{
 		return SNew(SSettingsEditor, Model);
 	}
 
-	virtual ISettingsEditorModelRef CreateModel( const ISettingsContainerRef& SettingsContainer ) override
+	virtual ISettingsEditorModelRef CreateModel( const TSharedRef<ISettingsContainer>& SettingsContainer ) override
 	{
 		return MakeShareable(new FSettingsEditorModel(SettingsContainer));
 	}
diff --git a/Engine/Source/Developer/SettingsEditor/Private/SettingsEditorPrivatePCH.h b/Engine/Source/Developer/SettingsEditor/Private/SettingsEditorPrivatePCH.h
--- a/Engine/Source/Developer/SettingsEditor/Private/SettingsEditorPrivatePCH.h
+++ b/Engine/Source/Developer/SettingsEditor/Private/SettingsEditorPrivatePCH.h
@@ -3,21 +3,22 @@
 #pragma once
 
 
-#include "SettingsEditor.h"
-
-
 /* Dependencies
  *****************************************************************************/
 
-#include "DesktopPlatformModule.h"
-#include "ModuleManager.h"
-#include "PropertyEditing.h"
+#include "Core.h"
+#include "CoreUObject.h"
+#include "EditorStyle.h"
 #include "ISettingsCategory.h"
+#include "ISettingsContainer.h"
+#include "ISettingsEditorModel.h"
+#include "ISettingsSection.h"
 #include "ISourceControlModule.h"
+#include "SlateBasics.h"
 
 
 /* Private includes
  *****************************************************************************/
 
 #include "SettingsEditorModel.h"
-
+#include "SSettingsEditor.h"
diff --git a/Engine/Source/Developer/SettingsEditor/Private/Widgets/SSettingsEditor.cpp b/Engine/Source/Developer/SettingsEditor/Private/Widgets/SSettingsEditor.cpp
--- a/Engine/Source/Developer/SettingsEditor/Private/Widgets/SSettingsEditor.cpp
+++ b/Engine/Source/Developer/SettingsEditor/Private/Widgets/SSettingsEditor.cpp
@@ -1,14 +1,16 @@
 // Copyright 1998-2014 Epic Games, Inc. All Rights Reserved.
 
 #include "SettingsEditorPrivatePCH.h"
-#include "SSettingsEditor.h"
-#include "EngineAnalytics.h"
 #include "AnalyticsEventAttribute.h"
+#include "DesktopPlatformModule.h"
+#include "EngineAnalytics.h"
 #include "IAnalyticsProvider.h"
-#include "SSettingsEditorCheckoutNotice.h"
+#include "NotificationManager.h"
+#include "PropertyEditing.h"
 #include "SHyperlink.h"
 #include "SNotificationList.h"
-#include "NotificationManager.h"
+#include "SSettingsEditorCheckoutNotice.h"
+
 
 #define LOCTEXT_NAMESPACE "SSettingsEditor"
 
@@ -16,7 +18,7 @@
 /* SSettingsEditor structors
  *****************************************************************************/
 
-SSettingsEditor::~SSettingsEditor( )
+SSettingsEditor::~SSettingsEditor()
 {
 	Model->OnSelectionChanged().RemoveAll(this);
 	SettingsContainer->OnCategoryModified().RemoveAll(this);
@@ -51,11 +53,10 @@ void SSettingsEditor::Construct( const FArguments& InArgs, const ISettingsEditor
 	SettingsView->SetIsPropertyEditingEnabledDelegate(FIsPropertyEditingEnabled::CreateSP(this, &SSettingsEditor::HandleSettingsViewEnabled));
 
 	// Create the watcher widget for the default config file (checks file status / SCC state)
-	FileWatcherWidget =
-		SNew(SSettingsEditorCheckoutNotice)
-		.Visibility(this, &SSettingsEditor::HandleDefaultConfigNoticeVisibility)
-		.OnFileProbablyModifiedExternally(this, &SSettingsEditor::ReloadConfigObject)
-		.ConfigFilePath(this, &SSettingsEditor::GetConfigFileName);
+	FileWatcherWidget = SNew(SSettingsEditorCheckoutNotice)
+		.Visibility(this, &SSettingsEditor::HandleCheckoutNoticeVisibility)
+		.OnFileProbablyModifiedExternally(this, &SSettingsEditor::HandleCheckoutNoticeFileProbablyModifiedExternally)
+		.ConfigFilePath(this, &SSettingsEditor::HandleCheckoutNoticeConfigFilePath);
 
 	ChildSlot
 	[
@@ -147,7 +148,7 @@ void SSettingsEditor::Construct( const FArguments& InArgs, const ISettingsEditor
 									[
 										// set as default button
 										SNew(SButton)
-											.Visibility(this, &SSettingsEditor::EditingNonDefaultSettingsVisibility)
+											.Visibility(this, &SSettingsEditor::HandleSetAsDefaultButtonVisibility)
 											.IsEnabled(this, &SSettingsEditor::HandleSetAsDefaultButtonEnabled)
 											.OnClicked(this, &SSettingsEditor::HandleSetAsDefaultButtonClicked)
 											.Text(LOCTEXT("SaveDefaultsButtonText", "Set as Default"))
@@ -190,7 +191,7 @@ void SSettingsEditor::Construct( const FArguments& InArgs, const ISettingsEditor
 									[
 										// reset defaults button
 										SNew(SButton)
-											.Visibility(this, &SSettingsEditor::EditingNonDefaultSettingsVisibility)
+											.Visibility(this, &SSettingsEditor::HandleSetAsDefaultButtonVisibility)
 											.IsEnabled(this, &SSettingsEditor::HandleResetToDefaultsButtonEnabled)
 											.OnClicked(this, &SSettingsEditor::HandleResetDefaultsButtonClicked)
 											.Text(LOCTEXT("ResetDefaultsButtonText", "Reset to Defaults"))
@@ -250,7 +251,7 @@ void SSettingsEditor::NotifyPostChange( const FPropertyChangedEvent& PropertyCha
 /* SSettingsEditor implementation
  *****************************************************************************/
 
-bool SSettingsEditor::CheckOutDefaultConfigFile( )
+bool SSettingsEditor::CheckOutDefaultConfigFile()
 {
 	TWeakObjectPtr<UObject> SettingsObject = GetSelectedSettingsObject();
 
@@ -317,7 +318,7 @@ FString SSettingsEditor::GetDefaultConfigFilePath( const TWeakObjectPtr<UObject>
 }
 
 
-TWeakObjectPtr<UObject> SSettingsEditor::GetSelectedSettingsObject( ) const
+TWeakObjectPtr<UObject> SSettingsEditor::GetSelectedSettingsObject() const
 {
 	ISettingsSectionPtr SelectedSection = Model->GetSelectedSection();
 
@@ -330,6 +331,21 @@ TWeakObjectPtr<UObject> SSettingsEditor::GetSelectedSettingsObject( ) const
 }
 
 
+bool SSettingsEditor::IsDefaultConfigCheckOutNeeded() const
+{
+	TWeakObjectPtr<UObject> SettingsObject = GetSelectedSettingsObject();
+
+	if (SettingsObject.IsValid() && SettingsObject->GetClass()->HasAnyClassFlags(CLASS_Config | CLASS_DefaultConfig))
+	{
+		return !FileWatcherWidget->IsUnlocked();
+	}
+	else
+	{
+		return false;
+	}
+}
+
+
 TSharedRef<SWidget> SSettingsEditor::MakeCategoryWidget( const ISettingsCategoryRef& Category )
 {
 	// create section widgets
@@ -419,7 +435,7 @@ TSharedRef<SWidget> SSettingsEditor::MakeCategoryWidget( const ISettingsCategory
 }
 
 
-bool SSettingsEditor::MakeDefaultConfigFileWritable( )
+bool SSettingsEditor::MakeDefaultConfigFileWritable()
 {
 	TWeakObjectPtr<UObject> SettingsObject = GetSelectedSettingsObject();
 
@@ -434,7 +450,23 @@ bool SSettingsEditor::MakeDefaultConfigFileWritable( )
 }
 
 
-void SSettingsEditor::ReloadCategories( )
+void SSettingsEditor::RecordPreferenceChangedAnalytics( ISettingsSectionPtr SelectedSection, const FPropertyChangedEvent& PropertyChangedEvent ) const
+{
+	UProperty* ChangedProperty = PropertyChangedEvent.MemberProperty;
+	// submit analytics data
+	if(FEngineAnalytics::IsAvailable() && ChangedProperty != nullptr)
+	{
+		TArray<FAnalyticsEventAttribute> EventAttributes;
+		EventAttributes.Add(FAnalyticsEventAttribute(TEXT("PropertySection"), SelectedSection->GetDisplayName().ToString()));
+		EventAttributes.Add(FAnalyticsEventAttribute(TEXT("PropertyClass"), ChangedProperty->GetOwnerClass()->GetName()));
+		EventAttributes.Add(FAnalyticsEventAttribute(TEXT("PropertyName"), ChangedProperty->GetName()));
+
+		FEngineAnalytics::GetProvider().RecordEvent(TEXT("Editor.Usage.PreferencesChanged"), EventAttributes);
+	}
+}
+
+
+void SSettingsEditor::ReloadCategories()
 {
 	CategoriesBox->ClearChildren();
 
@@ -466,44 +498,33 @@ void SSettingsEditor::ShowNotification( const FText& Text, SNotificationItem::EC
 /* SSettingsEditor callbacks
  *****************************************************************************/
 
-
-void SSettingsEditor::HandleCultureChanged( )
+FString SSettingsEditor::HandleCheckoutNoticeConfigFilePath() const
 {
-	ReloadCategories();
-}
-
-
+	TWeakObjectPtr<UObject> SettingsObject = GetSelectedSettingsObject();
 
-void SSettingsEditor::RecordPreferenceChangedAnalytics( ISettingsSectionPtr SelectedSection, const FPropertyChangedEvent& PropertyChangedEvent ) const
-{
-	UProperty* ChangedProperty = PropertyChangedEvent.MemberProperty;
-	// submit analytics data
-	if(FEngineAnalytics::IsAvailable() && ChangedProperty != nullptr)
+	if (SettingsObject.IsValid() && SettingsObject->GetClass()->HasAnyClassFlags(CLASS_Config | CLASS_DefaultConfig))
 	{
-		TArray<FAnalyticsEventAttribute> EventAttributes;
-		EventAttributes.Add(FAnalyticsEventAttribute(TEXT("PropertySection"), SelectedSection->GetDisplayName().ToString()));
-		EventAttributes.Add(FAnalyticsEventAttribute(TEXT("PropertyClass"), ChangedProperty->GetOwnerClass()->GetName()));
-		EventAttributes.Add(FAnalyticsEventAttribute(TEXT("PropertyName"), ChangedProperty->GetName()));
-
-		FEngineAnalytics::GetProvider().RecordEvent(TEXT("Editor.Usage.PreferencesChanged"), EventAttributes);
+		return GetDefaultConfigFilePath(SettingsObject);
+	}
+	else
+	{
+		return FString();
 	}
 }
 
-FString SSettingsEditor::GetConfigFileName() const
+
+void SSettingsEditor::HandleCheckoutNoticeFileProbablyModifiedExternally()
 {
 	TWeakObjectPtr<UObject> SettingsObject = GetSelectedSettingsObject();
 
 	if (SettingsObject.IsValid() && SettingsObject->GetClass()->HasAnyClassFlags(CLASS_Config | CLASS_DefaultConfig))
 	{
-		return GetDefaultConfigFilePath(SettingsObject);
-	}
-	else
-	{
-		return FString();
+		SettingsObject->ReloadConfig();
 	}
 }
 
-EVisibility SSettingsEditor::HandleDefaultConfigNoticeVisibility( ) const
+
+EVisibility SSettingsEditor::HandleCheckoutNoticeVisibility() const
 {
 	TWeakObjectPtr<UObject> SettingsObject = GetSelectedSettingsObject();
 
@@ -516,7 +537,13 @@ EVisibility SSettingsEditor::HandleDefaultConfigNoticeVisibility( ) const
 }
 
 
-FReply SSettingsEditor::HandleExportButtonClicked( )
+void SSettingsEditor::HandleCultureChanged()
+{
+	ReloadCategories();
+}
+
+
+FReply SSettingsEditor::HandleExportButtonClicked()
 {
 	ISettingsSectionPtr SelectedSection = Model->GetSelectedSection();
 
@@ -550,7 +577,7 @@ FReply SSettingsEditor::HandleExportButtonClicked( )
 }
 
 
-bool SSettingsEditor::HandleExportButtonEnabled( ) const
+bool SSettingsEditor::HandleExportButtonEnabled() const
 {
 	ISettingsSectionPtr SelectedSection = Model->GetSelectedSection();
 
@@ -563,7 +590,7 @@ bool SSettingsEditor::HandleExportButtonEnabled( ) const
 }
 
 
-FReply SSettingsEditor::HandleImportButtonClicked( )
+FReply SSettingsEditor::HandleImportButtonClicked()
 {
 	ISettingsSectionPtr SelectedSection = Model->GetSelectedSection();
 
@@ -591,7 +618,7 @@ FReply SSettingsEditor::HandleImportButtonClicked( )
 }
 
 
-bool SSettingsEditor::HandleImportButtonEnabled( ) const
+bool SSettingsEditor::HandleImportButtonEnabled() const
 {
 	ISettingsSectionPtr SelectedSection = Model->GetSelectedSection();
 
@@ -604,7 +631,7 @@ bool SSettingsEditor::HandleImportButtonEnabled( ) const
 }
 
 
-void SSettingsEditor::HandleModelSelectionChanged( )
+void SSettingsEditor::HandleModelSelectionChanged()
 {
 	ISettingsSectionPtr SelectedSection = Model->GetSelectedSection();
 
@@ -653,7 +680,7 @@ void SSettingsEditor::HandleModelSelectionChanged( )
 }
 
 
-FReply SSettingsEditor::HandleResetDefaultsButtonClicked( )
+FReply SSettingsEditor::HandleResetDefaultsButtonClicked()
 {
 	ISettingsSectionPtr SelectedSection = Model->GetSelectedSection();
 
@@ -666,7 +693,7 @@ FReply SSettingsEditor::HandleResetDefaultsButtonClicked( )
 }
 
 
-bool SSettingsEditor::HandleResetToDefaultsButtonEnabled( ) const
+bool SSettingsEditor::HandleResetToDefaultsButtonEnabled() const
 {
 	ISettingsSectionPtr SelectedSection = Model->GetSelectedSection();
 
@@ -679,7 +706,7 @@ bool SSettingsEditor::HandleResetToDefaultsButtonEnabled( ) const
 }
 
 
-EVisibility SSettingsEditor::EditingNonDefaultSettingsVisibility( ) const
+EVisibility SSettingsEditor::HandleSetAsDefaultButtonVisibility() const
 {
 	ISettingsSectionPtr SelectedSection = Model->GetSelectedSection();
 
@@ -704,7 +731,7 @@ EVisibility SSettingsEditor::HandleSectionLinkImageVisibility( ISettingsSectionP
 }
 
 
-FReply SSettingsEditor::HandleSetAsDefaultButtonClicked( )
+FReply SSettingsEditor::HandleSetAsDefaultButtonClicked()
 {
 	ISettingsSectionPtr SelectedSection = Model->GetSelectedSection();
 
@@ -741,7 +768,7 @@ FReply SSettingsEditor::HandleSetAsDefaultButtonClicked( )
 }
 
 
-bool SSettingsEditor::HandleSetAsDefaultButtonEnabled( ) const
+bool SSettingsEditor::HandleSetAsDefaultButtonEnabled() const
 {
 	ISettingsSectionPtr SelectedSection = Model->GetSelectedSection();
 
@@ -754,7 +781,7 @@ bool SSettingsEditor::HandleSetAsDefaultButtonEnabled( ) const
 }
 
 
-FText SSettingsEditor::HandleSettingsBoxDescriptionText( ) const
+FText SSettingsEditor::HandleSettingsBoxDescriptionText() const
 {
 	ISettingsSectionPtr SelectedSection = Model->GetSelectedSection();
 
@@ -767,7 +794,7 @@ FText SSettingsEditor::HandleSettingsBoxDescriptionText( ) const
 }
 
 
-FString SSettingsEditor::HandleSettingsBoxTitleText( ) const
+FString SSettingsEditor::HandleSettingsBoxTitleText() const
 {
 	ISettingsSectionPtr SelectedSection = Model->GetSelectedSection();
 
@@ -780,7 +807,7 @@ FString SSettingsEditor::HandleSettingsBoxTitleText( ) const
 }
 
 
-EVisibility SSettingsEditor::HandleSettingsBoxVisibility( ) const
+EVisibility SSettingsEditor::HandleSettingsBoxVisibility() const
 {
 	ISettingsSectionPtr SelectedSection = Model->GetSelectedSection();
 
@@ -799,7 +826,7 @@ void SSettingsEditor::HandleSettingsContainerCategoryModified( const FName& Cate
 }
 
 
-bool SSettingsEditor::HandleSettingsViewEnabled( ) const
+bool SSettingsEditor::HandleSettingsViewEnabled() const
 {
 	ISettingsSectionPtr SelectedSection = Model->GetSelectedSection();
 
@@ -807,7 +834,7 @@ bool SSettingsEditor::HandleSettingsViewEnabled( ) const
 }
 
 
-EVisibility SSettingsEditor::HandleSettingsViewVisibility( ) const
+EVisibility SSettingsEditor::HandleSettingsViewVisibility() const
 {
 	ISettingsSectionPtr SelectedSection = Model->GetSelectedSection();
 
@@ -820,28 +847,4 @@ EVisibility SSettingsEditor::HandleSettingsViewVisibility( ) const
 }
 
 
-// Do we need to edit the default config file?
-bool SSettingsEditor::IsDefaultConfigCheckOutNeeded() const
-{
-	TWeakObjectPtr<UObject> SettingsObject = GetSelectedSettingsObject();
-	if (SettingsObject.IsValid() && SettingsObject->GetClass()->HasAnyClassFlags(CLASS_Config | CLASS_DefaultConfig))
-	{
-		return !FileWatcherWidget->IsUnlocked();
-	}
-	else
-	{
-		return false;
-	}
-}
-
-void SSettingsEditor::ReloadConfigObject()
-{
-	TWeakObjectPtr<UObject> SettingsObject = GetSelectedSettingsObject();
-
-	if (SettingsObject.IsValid() && SettingsObject->GetClass()->HasAnyClassFlags(CLASS_Config | CLASS_DefaultConfig))
-	{
-		SettingsObject->ReloadConfig();
-	}
-}
-
 #undef LOCTEXT_NAMESPACE
diff --git a/Engine/Source/Developer/SettingsEditor/Private/Widgets/SSettingsEditor.h b/Engine/Source/Developer/SettingsEditor/Private/Widgets/SSettingsEditor.h
--- a/Engine/Source/Developer/SettingsEditor/Private/Widgets/SSettingsEditor.h
+++ b/Engine/Source/Developer/SettingsEditor/Private/Widgets/SSettingsEditor.h
@@ -1,9 +1,19 @@
 // Copyright 1998-2014 Epic Games, Inc. All Rights Reserved.
 
 #pragma once
+
+#include "NotifyHook.h"
 #include "SNotificationList.h"
 
 
+// forward declarations
+class FEditPropertyChain;
+struct FPropertyChangedEvent;
+class IDetailsView;
+class SSettingsEditorCheckoutNotice;
+class UObject;
+
+
 /**
  * Implements an editor widget for settings.
  */
@@ -18,10 +28,8 @@ public:
 
 public:
 
-	/**
-	 * Destructor.
-	 */
-	~SSettingsEditor( );
+	/** Destructor. */
+	~SSettingsEditor();
 
 public:
 
@@ -46,7 +54,7 @@ protected:
 	 *
 	 * @return true if the check-out succeeded, false otherwise.
 	 */
-	bool CheckOutDefaultConfigFile( );
+	bool CheckOutDefaultConfigFile();
 
 	/**
 	 * Gets the absolute path to the Default.ini for the specified object.
@@ -61,7 +69,14 @@ protected:
 	 *
 	 * @return The settings object.
 	 */
-	TWeakObjectPtr<UObject> GetSelectedSettingsObject( ) const;
+	TWeakObjectPtr<UObject> GetSelectedSettingsObject() const;
+
+	/**
+	 * Checks whether the default config file needs to be checked out for editing.
+	 *
+	 * @return true if the file needs to be checked out, false otherwise.
+	 */
+	bool IsDefaultConfigCheckOutNeeded() const;
 
 	/**
 	 * Creates a widget for the given settings category.
@@ -76,12 +91,18 @@ protected:
 	 *
 	 * @return true if it was made writable, false otherwise.
 	 */
-	bool MakeDefaultConfigFileWritable( );
+	bool MakeDefaultConfigFileWritable();
 
 	/**
-	 * Reloads the settings categories.
+	 * Reports a preference changed event to the analytics system.
+	 *
+	 * @param SelectedSection The currently selected settings section.
+	 * @param PropertyChangedEvent The event for the changed property.
 	 */
-	void ReloadCategories( );
+	void RecordPreferenceChangedAnalytics( ISettingsSectionPtr SelectedSection, const FPropertyChangedEvent& PropertyChangedEvent ) const;
+
+	/** Reloads the settings categories. */
+	void ReloadCategories();
 
 	/**
 	 * Shows a notification pop-up.
@@ -93,100 +114,95 @@ protected:
 
 private:
 
-	// Callback for clicking the Back link.
-	FReply HandleBackButtonClicked( );
+	/** Callback for clicking the Back link. */
+	FReply HandleBackButtonClicked();
+
+	/** Returns the config file name currently being edited. */
+	FString HandleCheckoutNoticeConfigFilePath() const;
 
-	// Callback for when the user's culture has changed.
-	void HandleCultureChanged( );
+	/** Reloads the configuration object. */
+	void HandleCheckoutNoticeFileProbablyModifiedExternally();
 
-	// Callback for determining the visibility of the 'Locked' notice.
-	EVisibility HandleDefaultConfigNoticeVisibility( ) const;
+	/** Callback for determining the visibility of the 'Locked' notice. */
+	EVisibility HandleCheckoutNoticeVisibility() const;
 
-	// Callback for clicking the 'Reset to Defaults' button.
-	FReply HandleExportButtonClicked( );
+	/** Callback for when the user's culture has changed. */
+	void HandleCultureChanged();
 
-	// Callback for getting the enabled state of the 'Export' button.
-	bool HandleExportButtonEnabled( ) const;
+	/** Callback for clicking the 'Reset to Defaults' button. */
+	FReply HandleExportButtonClicked();
 
-	// Callback for clicking the 'Reset to Defaults' button.
-	FReply HandleImportButtonClicked( );
+	/** Callback for getting the enabled state of the 'Export' button. */
+	bool HandleExportButtonEnabled() const;
 
-	// Callback for getting the enabled state of the 'Import' button.
-	bool HandleImportButtonEnabled( ) const;
+	/** Callback for clicking the 'Reset to Defaults' button. */
+	FReply HandleImportButtonClicked();
 
-	// Callback for changing the selected settings section.
-	void HandleModelSelectionChanged( );
+	/** Callback for getting the enabled state of the 'Import' button. */
+	bool HandleImportButtonEnabled() const;
 
-	// Callback for clicking the 'Reset to Defaults' button.
-	FReply HandleResetDefaultsButtonClicked( );
+	/** Callback for changing the selected settings section. */
+	void HandleModelSelectionChanged();
 
-	// Callback for getting the enabled state of the 'Reset to Defaults' button.
-	bool HandleResetToDefaultsButtonEnabled( ) const;
+	/** Callback for clicking the 'Reset to Defaults' button. */
+	FReply HandleResetDefaultsButtonClicked();
 
-	// Callback for getting the visibility of options that are only visible when editing something that came from a non-default object
-	EVisibility EditingNonDefaultSettingsVisibility( ) const;
+	/** Callback for getting the enabled state of the 'Reset to Defaults' button. */
+	bool HandleResetToDefaultsButtonEnabled() const;
 
-	// Callback for navigating a settings section link.
+	/** Callback for navigating a settings section link. */
 	void HandleSectionLinkNavigate( ISettingsSectionPtr Section );
 
-	// Callback for getting the visibility of a section link image.
+	/** Callback for getting the visibility of a section link image. */
 	EVisibility HandleSectionLinkImageVisibility( ISettingsSectionPtr Section ) const;
 
-	// Callback for clicking the 'Set as Defaults' button.
-	FReply HandleSetAsDefaultButtonClicked( );
-
-	// Callback for getting the enabled state of the 'Set as Defaults' button.
-	bool HandleSetAsDefaultButtonEnabled( ) const;
+	/** Callback for clicking the 'Set as Defaults' button. */
+	FReply HandleSetAsDefaultButtonClicked();
 
-	// Callback for getting the section description text.
-	FText HandleSettingsBoxDescriptionText( ) const;
+	/** Callback for getting the enabled state of the 'Set as Defaults' button. */
+	bool HandleSetAsDefaultButtonEnabled() const;
 
-	// Callback for getting the section text for the settings box.
-	FString HandleSettingsBoxTitleText( ) const;
+	/** Callback for getting the visibility of options that are only visible when editing something that came from a non-default object. */
+	EVisibility HandleSetAsDefaultButtonVisibility() const;
 
-	// Callback for determining the visibility of the settings box.
-	EVisibility HandleSettingsBoxVisibility( ) const;
+	/** Callback for getting the section description text. */
+	FText HandleSettingsBoxDescriptionText() const;
 
-	// Callback for the modification of categories in the settings container.
-	void HandleSettingsContainerCategoryModified( const FName& CategoryName );
-
-	// Callback for checking whether the settings view is enabled.
-	bool HandleSettingsViewEnabled( ) const;
+	/** Callback for getting the section text for the settings box. */
+	FString HandleSettingsBoxTitleText() const;
 
-	// Callback for determining the visibility of the settings view.
-	EVisibility HandleSettingsViewVisibility( ) const;
+	/** Callback for determining the visibility of the settings box. */
+	EVisibility HandleSettingsBoxVisibility() const;
 
-	// Reports a preference changed event to the analytics system */
-	void RecordPreferenceChangedAnalytics( ISettingsSectionPtr SelectedSection, const FPropertyChangedEvent& PropertyChangedEvent ) const;
-
-	// Returns the config file name currently being edited
-	FString GetConfigFileName() const;
+	/** Callback for the modification of categories in the settings container. */
+	void HandleSettingsContainerCategoryModified( const FName& CategoryName );
 
-	// Do we need to edit the default config file?
-	bool IsDefaultConfigCheckOutNeeded() const;
+	/** Callback for checking whether the settings view is enabled. */
+	bool HandleSettingsViewEnabled() const;
 
-	void ReloadConfigObject();
+	/** Callback for determining the visibility of the settings view. */
+	EVisibility HandleSettingsViewVisibility() const;
 
 private:
 
-	// Watcher widget for the default config file (checks file status / SCC state)
-	TSharedPtr<class SSettingsEditorCheckoutNotice> FileWatcherWidget;
-
-	// Holds the vertical box for settings categories.
+	/** Holds the vertical box for settings categories. */
 	TSharedPtr<SVerticalBox> CategoriesBox;
 
-	// Holds the overlay slot for custom widgets.
+	/** Holds the overlay slot for custom widgets. */
 	SOverlay::FOverlaySlot* CustomWidgetSlot;
 
-	// Holds the path to the directory that the last settings were exported to.
+	/** Watcher widget for the default config file (checks file status / SCC state). */
+	TSharedPtr<SSettingsEditorCheckoutNotice> FileWatcherWidget;
+
+	/** Holds the path to the directory that the last settings were exported to. */
 	FString LastExportDir;
 
-	// Holds a pointer to the view model.
+	/** Holds a pointer to the view model. */
 	ISettingsEditorModelPtr Model;
 
-	// Holds the settings container.
+	/** Holds the settings container. */
 	ISettingsContainerPtr SettingsContainer;
 
-	// Holds the details view.
+	/** Holds the details view. */
 	TSharedPtr<IDetailsView> SettingsView;
 };
diff --git a/Engine/Source/Developer/SettingsEditor/Public/ISettingsEditorModel.h b/Engine/Source/Developer/SettingsEditor/Public/ISettingsEditorModel.h
new file mode 100644
--- /dev/null
+++ b/Engine/Source/Developer/SettingsEditor/Public/ISettingsEditorModel.h
@@ -0,0 +1,65 @@
+// Copyright 1998-2014 Epic Games, Inc. All Rights Reserved.
+
+#pragma once
+
+
+// forward declarations
+class ISettingsSection;
+
+
+/**
+ * Interface for settings editor view models.
+ *
+ * The settings editor view model stores the view state for the Settings Editor UI.
+ * Instances of this interface can be passed to Settings Editors in order to provide
+ * access to the settings container being added and to perform various user actions,
+ * such as setting the currently selected settings section.
+ */
+class ISettingsEditorModel
+{
+public:
+
+	/**
+	 * Gets the selected settings section.
+	 *
+	 * @return The selected section, if any.
+	 */
+	virtual const TSharedPtr<ISettingsSection>& GetSelectedSection() const = 0;
+
+	/**
+	 * Gets the settings container.
+	 *
+	 * @return The settings container.
+	 */
+	virtual const TSharedRef<ISettingsContainer>& GetSettingsContainer() const = 0;
+
+	/**
+	 * Selects the specified settings section to be displayed in the editor.
+	 *
+	 * @param Section The section to select.
+	 *
+	 * @see SelectPreviousSection
+	 */
+	virtual void SelectSection( const TSharedPtr<ISettingsSection>& Section ) = 0;
+
+public:
+
+	/**
+	 * Returns a delegate that is executed when the selected settings section has changed.
+	 *
+	 * @return The delegate.
+	 */
+	virtual FSimpleMulticastDelegate& OnSelectionChanged() = 0;
+
+public:
+
+	/** Virtual destructor. */
+	virtual ~ISettingsEditorModel() { }
+};
+
+
+/** Type definition for shared pointers to instances of ISettingsEditorModel. */
+typedef TSharedPtr<ISettingsEditorModel> ISettingsEditorModelPtr;
+
+/** Type definition for shared references to instances of ISettingsEditorModel. */
+typedef TSharedRef<ISettingsEditorModel> ISettingsEditorModelRef;
diff --git a/Engine/Source/Developer/SettingsEditor/Public/ISettingsEditorModule.h b/Engine/Source/Developer/SettingsEditor/Public/ISettingsEditorModule.h
new file mode 100644
--- /dev/null
+++ b/Engine/Source/Developer/SettingsEditor/Public/ISettingsEditorModule.h
@@ -0,0 +1,44 @@
+// Copyright 1998-2014 Epic Games, Inc. All Rights Reserved.
+
+#pragma once
+
+#include "ModuleInterface.h"
+
+
+// forward declarations
+class ISettingsContainer;
+class ISettingsEditorModel;
+class SWidget;
+
+
+/**
+ * Interface for settings editor modules.
+ */
+class ISettingsEditorModule
+	: public IModuleInterface
+{
+public:
+
+	/**
+	 * Creates a settings editor widget.
+	 *
+	 * @param Model The view model.
+	 * @return The new widget.
+	 * @see CreateCodel
+	 */
+	virtual TSharedRef<SWidget> CreateEditor( const TSharedRef<ISettingsEditorModel>& Model ) = 0;
+
+	/**
+	 * Creates a view model for the settings editor widget.
+	 *
+	 * @param SettingsContainer The settings container.
+	 * @return The controller.
+	 * @see CreateEditor
+	 */
+	virtual ISettingsEditorModelRef CreateModel( const TSharedRef<ISettingsContainer>& SettingsContainer ) = 0;
+
+public:
+
+	/** Virtual destructor. */
+	virtual ~ISettingsEditorModule() { }
+};
diff --git a/Engine/Source/Developer/SettingsEditor/Public/Interfaces/ISettingsEditorModel.h b/Engine/Source/Developer/SettingsEditor/Public/Interfaces/ISettingsEditorModel.h
deleted file mode 100644
diff --git a/Engine/Source/Developer/SettingsEditor/Public/Interfaces/ISettingsEditorModule.h b/Engine/Source/Developer/SettingsEditor/Public/Interfaces/ISettingsEditorModule.h
deleted file mode 100644
diff --git a/Engine/Source/Developer/SettingsEditor/Public/SettingsEditor.h b/Engine/Source/Developer/SettingsEditor/Public/SettingsEditor.h
deleted file mode 100644
diff --git a/Engine/Source/Developer/SettingsEditor/SettingsEditor.Build.cs b/Engine/Source/Developer/SettingsEditor/SettingsEditor.Build.cs
--- a/Engine/Source/Developer/SettingsEditor/SettingsEditor.Build.cs
+++ b/Engine/Source/Developer/SettingsEditor/SettingsEditor.Build.cs
@@ -8,8 +8,7 @@ public class SettingsEditor : ModuleRules
 	{
 		PublicDependencyModuleNames.AddRange(
 			new string[] {
-                "EditorStyle",
-				"Slate",
+				"Core",
 			}
 		);
 
@@ -21,16 +20,17 @@ public class SettingsEditor : ModuleRules
 
 		PrivateDependencyModuleNames.AddRange(
             new string[] {
-				"Core",
+				"Analytics",
                 "CoreUObject",
-                "InputCore",
 				"DesktopPlatform",
+                "EditorStyle",
+				"Engine",
+                "InputCore",
 				"PropertyEditor",
+				"SharedSettingsWidgets",
+				"Slate",
 				"SlateCore",
 				"SourceControl",
-				"Analytics",
-				"SharedSettingsWidgets",
-				"Engine",
             }
         );
 
diff --git a/Engine/Source/Editor/EditorSettingsViewer/Private/EditorSettingsViewerModule.cpp b/Engine/Source/Editor/EditorSettingsViewer/Private/EditorSettingsViewerModule.cpp
--- a/Engine/Source/Editor/EditorSettingsViewer/Private/EditorSettingsViewerModule.cpp
+++ b/Engine/Source/Editor/EditorSettingsViewer/Private/EditorSettingsViewerModule.cpp
@@ -2,6 +2,9 @@
 
 #include "EditorSettingsViewerPrivatePCH.h"
 #include "ISettingsCategory.h"
+#include "ISettingsContainer.h"
+#include "ISettingsEditorModel.h"
+#include "ISettingsEditorModule.h"
 #include "ISettingsModule.h"
 #include "ISettingsSection.h"
 #include "ISettingsViewer.h"
@@ -256,7 +259,7 @@ private:
 
 			if (SettingsContainer.IsValid())
 			{
-				ISettingsEditorModule& SettingsEditorModule = ISettingsEditorModule::GetRef();
+				ISettingsEditorModule& SettingsEditorModule = FModuleManager::GetModuleChecked<ISettingsEditorModule>("SettingsEditor");
 				ISettingsEditorModelRef SettingsEditorModel = SettingsEditorModule.CreateModel(SettingsContainer.ToSharedRef());
 
 				SettingsEditor = SettingsEditorModule.CreateEditor(SettingsEditorModel);
diff --git a/Engine/Source/Editor/EditorSettingsViewer/Private/EditorSettingsViewerPrivatePCH.h b/Engine/Source/Editor/EditorSettingsViewer/Private/EditorSettingsViewerPrivatePCH.h
--- a/Engine/Source/Editor/EditorSettingsViewer/Private/EditorSettingsViewerPrivatePCH.h
+++ b/Engine/Source/Editor/EditorSettingsViewer/Private/EditorSettingsViewerPrivatePCH.h
@@ -11,4 +11,3 @@
 #include "InputBindingEditor.h"
 #include "InternationalizationSettingsModel.h"
 #include "MessageLog.h"
-#include "SettingsEditor.h"
diff --git a/Engine/Source/Editor/ProjectSettingsViewer/Private/ProjectSettingsViewerModule.cpp b/Engine/Source/Editor/ProjectSettingsViewer/Private/ProjectSettingsViewerModule.cpp
--- a/Engine/Source/Editor/ProjectSettingsViewer/Private/ProjectSettingsViewerModule.cpp
+++ b/Engine/Source/Editor/ProjectSettingsViewer/Private/ProjectSettingsViewerModule.cpp
@@ -3,15 +3,18 @@
 #include "ProjectSettingsViewerPrivatePCH.h"
 #include "Engine/Console.h"
 #include "ProjectTargetPlatformEditor.h"
-#include "CookerSettings.h"
 #include "ISettingsCategory.h"
+#include "ISettingsContainer.h"
+#include "ISettingsEditorModel.h"
+#include "ISettingsEditorModule.h"
 #include "ISettingsModule.h"
 #include "ISettingsViewer.h"
 #include "ModuleManager.h"
 #include "SDockTab.h"
 
-#include "Runtime/Engine/Classes/Engine/UserInterfaceSettings.h"
+#include "CookerSettings.h"
 #include "Runtime/Engine/Classes/Engine/RendererSettings.h"
+#include "Runtime/Engine/Classes/Engine/UserInterfaceSettings.h"
 #include "Runtime/Engine/Classes/Sound/AudioSettings.h"
 
 
@@ -51,7 +54,7 @@ public:
 
 	// IModuleInterface interface
 
-	virtual void StartupModule( ) override
+	virtual void StartupModule() override
 	{
 		ISettingsModule* SettingsModule = FModuleManager::GetModulePtr<ISettingsModule>("Settings");
 
@@ -68,13 +71,13 @@ public:
 			.SetMenuType(ETabSpawnerMenuType::Hide);
 	}
 
-	virtual void ShutdownModule( ) override
+	virtual void ShutdownModule() override
 	{
 		FGlobalTabmanager::Get()->UnregisterNomadTabSpawner(ProjectSettingsTabName);
 		UnregisterSettings();
 	}
 
-	virtual bool SupportsDynamicReloading( ) override
+	virtual bool SupportsDynamicReloading() override
 	{
 		return true;
 	}
@@ -239,7 +242,7 @@ protected:
 	}
 
 	/** Unregisters all settings. */
-	void UnregisterSettings( )
+	void UnregisterSettings()
 	{
 		ISettingsModule* SettingsModule = FModuleManager::GetModulePtr<ISettingsModule>("Settings");
 
@@ -282,7 +285,7 @@ private:
 
 			if (SettingsContainer.IsValid())
 			{
-				ISettingsEditorModule& SettingsEditorModule = ISettingsEditorModule::GetRef();
+				ISettingsEditorModule& SettingsEditorModule = FModuleManager::GetModuleChecked<ISettingsEditorModule>("SettingsEditor");
 				ISettingsEditorModelRef SettingsEditorModel = SettingsEditorModule.CreateModel(SettingsContainer.ToSharedRef());
 
 				SettingsEditor = SettingsEditorModule.CreateEditor(SettingsEditorModel);
diff --git a/Engine/Source/Editor/ProjectSettingsViewer/Private/ProjectSettingsViewerPrivatePCH.h b/Engine/Source/Editor/ProjectSettingsViewer/Private/ProjectSettingsViewerPrivatePCH.h
--- a/Engine/Source/Editor/ProjectSettingsViewer/Private/ProjectSettingsViewerPrivatePCH.h
+++ b/Engine/Source/Editor/ProjectSettingsViewer/Private/ProjectSettingsViewerPrivatePCH.h
@@ -10,5 +10,4 @@
 #include "EngineSettings.h"
 #include "ModuleManager.h"
 #include "MoviePlayerSettings.h"
-#include "SettingsEditor.h"
 #include "UnrealEd.h"
