Added a ConfigRestartRequired meta-data property to allow you to tag config properties as requring an application restart before they take effect

UE-9022 - IWCE: Had to restart editor after toggling Blueprintable Components option for anything to work

diff --git a/Engine/Source/Developer/SettingsEditor/Private/SettingsEditorModule.cpp b/Engine/Source/Developer/SettingsEditor/Private/SettingsEditorModule.cpp
--- a/Engine/Source/Developer/SettingsEditor/Private/SettingsEditorModule.cpp
+++ b/Engine/Source/Developer/SettingsEditor/Private/SettingsEditorModule.cpp
@@ -2,6 +2,98 @@
 
 #include "SettingsEditorPrivatePCH.h"
 #include "ISettingsEditorModule.h"
+#include "NotificationManager.h"
+
+
+#define LOCTEXT_NAMESPACE "SSettingsEditor"
+
+
+/** Manages the notification for when the application needs to be restarted due to a settings change */
+class FApplicationRestartRequiredNotification
+{
+public:
+	void SetOnRestartApplicationCallback( FSimpleDelegate InRestartApplicationDelegate )
+	{
+		RestartApplicationDelegate = InRestartApplicationDelegate;
+	}
+
+	void OnRestartRequired()
+	{
+		TSharedPtr<SNotificationItem> NotificationPin = NotificationPtr.Pin();
+		if (NotificationPin.IsValid() || !RestartApplicationDelegate.IsBound())
+		{
+			return;
+		}
+		
+		FNotificationInfo Info( LOCTEXT("RestartRequiredTitle", "Restart required to apply new settings") );
+
+		// Add the buttons with text, tooltip and callback
+		Info.ButtonDetails.Add(FNotificationButtonInfo(
+			LOCTEXT("RestartNow", "Restart Now"), 
+			LOCTEXT("RestartNowToolTip", "Restart now to finish applying your new settings."), 
+			FSimpleDelegate::CreateRaw(this, &FApplicationRestartRequiredNotification::OnRestartClicked))
+			);
+		Info.ButtonDetails.Add(FNotificationButtonInfo(
+			LOCTEXT("RestartLater", "Restart Later"), 
+			LOCTEXT("RestartLaterToolTip", "Dismiss this notificaton without restarting. Some new settings will not be applied."), 
+			FSimpleDelegate::CreateRaw(this, &FApplicationRestartRequiredNotification::OnDismissClicked))
+			);
+
+		// We will be keeping track of this ourselves
+		Info.bFireAndForget = false;
+
+		// Set the width so that the notification doesn't resize as its text changes
+		Info.WidthOverride = 300.0f;
+
+		Info.bUseLargeFont = false;
+		Info.bUseThrobber = false;
+		Info.bUseSuccessFailIcons = false;
+
+		// Launch notification
+		NotificationPtr = FSlateNotificationManager::Get().AddNotification(Info);
+		NotificationPin = NotificationPtr.Pin();
+
+		if (NotificationPin.IsValid())
+		{
+			NotificationPin->SetCompletionState(SNotificationItem::CS_Pending);
+		}
+	}
+
+
+
+private:
+	void OnRestartClicked()
+	{
+		TSharedPtr<SNotificationItem> NotificationPin = NotificationPtr.Pin();
+		if (NotificationPin.IsValid())
+		{
+			NotificationPin->SetText(LOCTEXT("RestartingNow", "Restarting..."));
+			NotificationPin->SetCompletionState(SNotificationItem::CS_Success);
+			NotificationPin->ExpireAndFadeout();
+			NotificationPtr.Reset();
+		}
+
+		RestartApplicationDelegate.ExecuteIfBound();
+	}
+
+	void OnDismissClicked()
+	{
+		TSharedPtr<SNotificationItem> NotificationPin = NotificationPtr.Pin();
+		if (NotificationPin.IsValid())
+		{
+			NotificationPin->SetText(LOCTEXT("RestartDismissed", "Restart Dismissed..."));
+			NotificationPin->SetCompletionState(SNotificationItem::CS_None);
+			NotificationPin->ExpireAndFadeout();
+			NotificationPtr.Reset();
+		}
+	}
+
+	/** Used to reference to the active restart notification */
+	TWeakPtr<SNotificationItem> NotificationPtr;
+
+	/** Used to actually restart the application */
+	FSimpleDelegate RestartApplicationDelegate;
+};
 
 
 /**
@@ -16,14 +108,32 @@ public:
 
 	virtual TSharedRef<SWidget> CreateEditor( const TSharedRef<ISettingsEditorModel>& Model ) override
 	{
-		return SNew(SSettingsEditor, Model);
+		return SNew(SSettingsEditor, Model)
+			.OnApplicationRestartRequired(FSimpleDelegate::CreateRaw(this, &FSettingsEditorModule::OnApplicationRestartRequired));
 	}
 
 	virtual ISettingsEditorModelRef CreateModel( const TSharedRef<ISettingsContainer>& SettingsContainer ) override
 	{
 		return MakeShareable(new FSettingsEditorModel(SettingsContainer));
 	}
+
+	virtual void OnApplicationRestartRequired() override
+	{
+		ApplicationRestartRequiredNotification.OnRestartRequired();
+	}
+
+	virtual void SetRestartApplicationCallback( FSimpleDelegate InRestartApplicationDelegate ) override
+	{
+		ApplicationRestartRequiredNotification.SetOnRestartApplicationCallback(InRestartApplicationDelegate);
+	}
+
+private:
+
+	FApplicationRestartRequiredNotification ApplicationRestartRequiredNotification;
 };
 
 
 IMPLEMENT_MODULE(FSettingsEditorModule, SettingsEditor);
+
+
+#undef LOCTEXT_NAMESPACE
diff --git a/Engine/Source/Developer/SettingsEditor/Private/Widgets/SSettingsEditor.cpp b/Engine/Source/Developer/SettingsEditor/Private/Widgets/SSettingsEditor.cpp
--- a/Engine/Source/Developer/SettingsEditor/Private/Widgets/SSettingsEditor.cpp
+++ b/Engine/Source/Developer/SettingsEditor/Private/Widgets/SSettingsEditor.cpp
@@ -33,6 +33,7 @@ void SSettingsEditor::Construct( const FArguments& InArgs, const ISettingsEditor
 {
 	Model = InModel;
 	SettingsContainer = InModel->GetSettingsContainer();
+	OnApplicationRestartRequiredDelegate = InArgs._OnApplicationRestartRequired;
 
 	// initialize settings view
 	FDetailsViewArgs DetailsViewArgs;
@@ -243,6 +244,12 @@ void SSettingsEditor::NotifyPostChange( const FPropertyChangedEvent& PropertyCha
 		RecordPreferenceChangedAnalytics( SelectedSection, PropertyChangedEvent );
 
 		SelectedSection->Save();
+
+		static const FName ConfigRestartRequiredKey = "ConfigRestartRequired";
+		if (PropertyChangedEvent.Property->GetBoolMetaData(ConfigRestartRequiredKey))
+		{
+			OnApplicationRestartRequiredDelegate.ExecuteIfBound();
+		}
 	}
 }
 
diff --git a/Engine/Source/Developer/SettingsEditor/Private/Widgets/SSettingsEditor.h b/Engine/Source/Developer/SettingsEditor/Private/Widgets/SSettingsEditor.h
--- a/Engine/Source/Developer/SettingsEditor/Private/Widgets/SSettingsEditor.h
+++ b/Engine/Source/Developer/SettingsEditor/Private/Widgets/SSettingsEditor.h
@@ -24,6 +24,7 @@ class SSettingsEditor
 public:
 
 	SLATE_BEGIN_ARGS(SSettingsEditor) { }
+		SLATE_EVENT( FSimpleDelegate, OnApplicationRestartRequired )
 	SLATE_END_ARGS()
 
 public:
@@ -205,4 +206,7 @@ private:
 
 	/** Holds the details view. */
 	TSharedPtr<IDetailsView> SettingsView;
+
+	/** Delegate called when this settings editor requests that the user be notified that the application needs to be restarted for some setting changes to take effect */
+	FSimpleDelegate OnApplicationRestartRequiredDelegate;
 };
diff --git a/Engine/Source/Developer/SettingsEditor/Public/ISettingsEditorModule.h b/Engine/Source/Developer/SettingsEditor/Public/ISettingsEditorModule.h
--- a/Engine/Source/Developer/SettingsEditor/Public/ISettingsEditorModule.h
+++ b/Engine/Source/Developer/SettingsEditor/Public/ISettingsEditorModule.h
@@ -24,7 +24,7 @@ public:
 	 *
 	 * @param Model The view model.
 	 * @return The new widget.
-	 * @see CreateCodel
+	 * @see CreateModel
 	 */
 	virtual TSharedRef<SWidget> CreateEditor( const TSharedRef<ISettingsEditorModel>& Model ) = 0;
 
@@ -35,7 +35,19 @@ public:
 	 * @return The controller.
 	 * @see CreateEditor
 	 */
-	virtual ISettingsEditorModelRef CreateModel( const TSharedRef<ISettingsContainer>& SettingsContainer ) = 0;
+	virtual TSharedRef<ISettingsEditorModel> CreateModel( const TSharedRef<ISettingsContainer>& SettingsContainer ) = 0;
+
+	/**
+	 * Called when the settings have been changed such that an application restart is required for them to be fully applied
+	 */
+	virtual void OnApplicationRestartRequired() = 0;
+
+	/**
+	 * Set the delegate that should be called when a setting editor needs to restart the application
+	 *
+	 * @param InRestartApplicationDelegate The new delegate to call
+	 */
+	virtual void SetRestartApplicationCallback( FSimpleDelegate InRestartApplicationDelegate ) = 0;
 
 public:
 
diff --git a/Engine/Source/Editor/EditorSettingsViewer/Private/EditorSettingsViewerModule.cpp b/Engine/Source/Editor/EditorSettingsViewer/Private/EditorSettingsViewerModule.cpp
--- a/Engine/Source/Editor/EditorSettingsViewer/Private/EditorSettingsViewerModule.cpp
+++ b/Engine/Source/Editor/EditorSettingsViewer/Private/EditorSettingsViewerModule.cpp
@@ -109,6 +109,7 @@ protected:
 			RegionAndLanguageSettingsSection->OnImport().BindRaw(this, &FEditorSettingsViewerModule::HandleRegionAndLanguageImport);
 			RegionAndLanguageSettingsSection->OnSaveDefaults().BindRaw(this, &FEditorSettingsViewerModule::HandleRegionAndLanguageSaveDefaults);
 			RegionAndLanguageSettingsSection->OnResetDefaults().BindRaw(this, &FEditorSettingsViewerModule::HandleRegionAndLanguageResetToDefault);
+			GetMutableDefault<UInternationalizationSettingsModel>()->OnSettingChanged().AddRaw(this, &FEditorSettingsViewerModule::HandleRegionAndLanguageSettingChanged);
 		}
 
 		// input bindings
@@ -239,6 +240,7 @@ protected:
 			SettingsModule->UnregisterSettings("Editor", "General", "GameAgnostic");
 			SettingsModule->UnregisterSettings("Editor", "General", "UserSettings");
 			SettingsModule->UnregisterSettings("Editor", "General", "AutomationTest");
+			SettingsModule->UnregisterSettings("Editor", "General", "Internationalization");
 			SettingsModule->UnregisterSettings("Editor", "General", "Experimental");
 
 			// level editor settings
@@ -402,6 +404,12 @@ private:
 		return true;
 	}
 
+	void HandleRegionAndLanguageSettingChanged()
+	{
+		ISettingsEditorModule& SettingsEditorModule = FModuleManager::GetModuleChecked<ISettingsEditorModule>("SettingsEditor");
+		SettingsEditorModule.OnApplicationRestartRequired();
+	}
+
 private:
 
 	/** Holds a pointer to the settings editor's view model. */
diff --git a/Engine/Source/Editor/UnrealEd/Classes/Settings/EditorExperimentalSettings.h b/Engine/Source/Editor/UnrealEd/Classes/Settings/EditorExperimentalSettings.h
--- a/Engine/Source/Editor/UnrealEd/Classes/Settings/EditorExperimentalSettings.h
+++ b/Engine/Source/Editor/UnrealEd/Classes/Settings/EditorExperimentalSettings.h
@@ -42,7 +42,7 @@ public:
 	bool bUnifiedBlueprintEditor;
 
 	/** Enable being able to subclass components in blueprints */
-	UPROPERTY(EditAnywhere, config, Category=Tools)
+	UPROPERTY(EditAnywhere, config, Category=Tools, meta=(ConfigRestartRequired=true))
 	bool bBlueprintableComponents;
 
 	/** The Messaging Debugger provides a visual utility for debugging the messaging system. */
@@ -78,7 +78,7 @@ public:
 	bool bContextMenuChunkAssignments;
 
 	/** Disable cook in the editor */
-	UPROPERTY(EditAnywhere, config, Category = Cooking, meta = (DisplayName = "Disable Cook In The Editor feature, requires editor restart (cooks from launch on will be run in a separate process if disabled)"))
+	UPROPERTY(EditAnywhere, config, Category = Cooking, meta = (DisplayName = "Disable Cook In The Editor feature (cooks from launch on will be run in a separate process if disabled)", ConfigRestartRequired=true))
 	bool bDisableCookInEditor;
 
 	/** Enable cook on the side */
diff --git a/Engine/Source/Editor/UnrealEd/Private/UnrealEdMisc.cpp b/Engine/Source/Editor/UnrealEd/Private/UnrealEdMisc.cpp
--- a/Engine/Source/Editor/UnrealEd/Private/UnrealEdMisc.cpp
+++ b/Engine/Source/Editor/UnrealEd/Private/UnrealEdMisc.cpp
@@ -20,6 +20,7 @@
 #include "AssetRegistryModule.h"
 #include "EngineAnalytics.h"
 #include "IAnalyticsProvider.h"
+#include "ISettingsEditorModule.h"
 #include "LevelEditor.h"
 #include "Toolkits/AssetEditorManager.h"
 #include "UObjectToken.h"
@@ -442,6 +443,10 @@ void FUnrealEdMisc::OnInit()
 
 	GEditor->GetTimerManager()->SetTimer( EditorAnalyticsHeartbeatTimerHandle, Delegate, Seconds, true );
 
+	// Give the settings editor a way to restart the editor when it needs to
+	ISettingsEditorModule& SettingsEditorModule = FModuleManager::GetModuleChecked<ISettingsEditorModule>("SettingsEditor");
+	SettingsEditorModule.SetRestartApplicationCallback(FSimpleDelegate::CreateRaw(this, &FUnrealEdMisc::RestartEditor, false));
+
 	// add handler to notify about navmesh building process
 	NavigationBuildingNotificationHandler = MakeShareable(new FNavigationBuildingNotificationImpl());
 
diff --git a/Engine/Source/Editor/UnrealEd/UnrealEd.Build.cs b/Engine/Source/Editor/UnrealEd/UnrealEd.Build.cs
--- a/Engine/Source/Editor/UnrealEd/UnrealEd.Build.cs
+++ b/Engine/Source/Editor/UnrealEd/UnrealEd.Build.cs
@@ -44,6 +44,7 @@ public class UnrealEd : ModuleRules
 				"NiagaraEditor",
 				"PlacementMode",
 				"Settings",
+				"SettingsEditor",
 				"SoundClassEditor",
 				"ViewportSnapping",
 				"SourceCodeAccess",
