Added confirmation message box to "Save as Defaults"

diff --git a/Engine/Source/Developer/SettingsEditor/Private/Widgets/SSettingsEditor.cpp b/Engine/Source/Developer/SettingsEditor/Private/Widgets/SSettingsEditor.cpp
--- a/Engine/Source/Developer/SettingsEditor/Private/Widgets/SSettingsEditor.cpp
+++ b/Engine/Source/Developer/SettingsEditor/Private/Widgets/SSettingsEditor.cpp
@@ -743,6 +743,11 @@ FReply SSettingsEditor::HandleSetAsDefaultButtonClicked()
 
 	if (SelectedSection.IsValid())
 	{
+		if (FMessageDialog::Open(EAppMsgType::YesNo, LOCTEXT("SaveAsDefaultUserConfirm", "Are you sure you want to update the default settings?")) != EAppReturnType::Yes)
+		{
+			return FReply::Handled();
+		}
+
 		if (IsDefaultConfigCheckOutNeeded())
 		{
 			if (ISourceControlModule::Get().IsEnabled())
