Editor - Fixing non-unity editor build, not including header file for DeveloperSettings.h

diff --git a/Engine/Source/Developer/SettingsEditor/Private/SettingsEditorModule.cpp b/Engine/Source/Developer/SettingsEditor/Private/SettingsEditorModule.cpp
--- a/Engine/Source/Developer/SettingsEditor/Private/SettingsEditorModule.cpp
+++ b/Engine/Source/Developer/SettingsEditor/Private/SettingsEditorModule.cpp
@@ -1,9 +1,11 @@
 // Copyright 1998-2015 Epic Games, Inc. All Rights Reserved.
 
 #include "SettingsEditorPrivatePCH.h"
+
 #include "ISettingsEditorModule.h"
 #include "NotificationManager.h"
 #include "ISettingsModule.h"
+#include "Engine/DeveloperSettings.h"
 
 #define LOCTEXT_NAMESPACE "SSettingsEditor"
 
