Changed the UDeveloperSettings auto discovery to ignore all abstract classes instead of just specifically UDeveloperSettings so that shared base classes don't generate blank spots in the settings list

diff --git a/Engine/Source/Developer/SettingsEditor/Private/SettingsEditorModule.cpp b/Engine/Source/Developer/SettingsEditor/Private/SettingsEditorModule.cpp
--- a/Engine/Source/Developer/SettingsEditor/Private/SettingsEditorModule.cpp
+++ b/Engine/Source/Developer/SettingsEditor/Private/SettingsEditorModule.cpp
@@ -222,8 +222,8 @@ private:
 				// Only Add the CDO of any UDeveloperSettings objects.
 				if ( Settings->HasAnyFlags(RF_ClassDefaultObject) && !Settings->GetClass()->HasAnyCastFlag(CLASS_Deprecated) )
 				{
-					// Ignore the setting if it's specifically the UDeveloperSettings
-					if ( Settings->GetClass() == UDeveloperSettings::StaticClass() )
+					// Ignore the setting if it's specifically the UDeveloperSettings or other abstract settings classes
+					if ( Settings->GetClass()->HasAnyClassFlags(CLASS_Abstract) )
 					{
 						continue;
 					}
