Add friendly reminder to Editor Preferences -> Set As Default messages that checking in the default editor settings file will affect other developers. Meant  to help prevent accidental stomping of other people's settings.

--------
Integrated using branch UE4-To-UE4-Orion (reversed) of change#2492063 by [REDACTED_PERSON] on 2015/03/26 10:35:55.

diff --git a/Engine/Source/Developer/SettingsEditor/Private/Widgets/SSettingsEditor.cpp b/Engine/Source/Developer/SettingsEditor/Private/Widgets/SSettingsEditor.cpp
--- a/Engine/Source/Developer/SettingsEditor/Private/Widgets/SSettingsEditor.cpp
+++ b/Engine/Source/Developer/SettingsEditor/Private/Widgets/SSettingsEditor.cpp
@@ -753,7 +753,7 @@ FReply SSettingsEditor::HandleSetAsDefaultButtonClicked()
 		{
 			if (ISourceControlModule::Get().IsEnabled())
 			{
-				if (FMessageDialog::Open(EAppMsgType::YesNo, LOCTEXT("SaveAsDefaultNeedsCheckoutMessage", "The default configuration file for these settings is currently not checked out. Would you like to check it out from source control?")) != EAppReturnType::Yes)
+				if (FMessageDialog::Open(EAppMsgType::YesNo, LOCTEXT("SaveAsDefaultNeedsCheckoutMessage", "The default configuration file for these settings is currently not checked out. Would you like to check it out from source control? \n\nIf checked into source control this would affect other developers.")) != EAppReturnType::Yes)
 				{
 					return FReply::Handled();
 				}
@@ -773,7 +773,7 @@ FReply SSettingsEditor::HandleSetAsDefaultButtonClicked()
 
 		SelectedSection->SaveDefaults();
 
-		FMessageDialog::Open(EAppMsgType::Ok, LOCTEXT("SaveAsDefaultsSucceededMessage", "The default configuration file for these settings was updated successfully."));
+		FMessageDialog::Open(EAppMsgType::Ok, LOCTEXT("SaveAsDefaultsSucceededMessage", "The default configuration file for these settings was updated successfully. \n\nIf checked into source control this would affect other developers."));
 	}
 
 	return FReply::Handled();
