UE-19491

Fixed an issue were Array Element changes were not being properly output to a default config.

diff --git a/Engine/Source/Developer/SettingsEditor/Private/Widgets/SSettingsEditor.cpp b/Engine/Source/Developer/SettingsEditor/Private/Widgets/SSettingsEditor.cpp
--- a/Engine/Source/Developer/SettingsEditor/Private/Widgets/SSettingsEditor.cpp
+++ b/Engine/Source/Developer/SettingsEditor/Private/Widgets/SSettingsEditor.cpp
@@ -244,7 +244,11 @@ void SSettingsEditor::NotifyPostChange( const FPropertyChangedEvent& PropertyCha
 	{
 		RecordPreferenceChangedAnalytics( SelectedSection, PropertyChangedEvent );
 
-		if (SelectedSection->GetSettingsObject()->GetClass()->HasAnyClassFlags(CLASS_DefaultConfig) && !PropertyChangedEvent.Property->IsA(UArrayProperty::StaticClass()))
+		// Determine if the Property is an Array or Array Element
+		UObject* Outer = PropertyChangedEvent.Property->GetOuter();
+		bool bIsArrayOrArrayElement = PropertyChangedEvent.Property->IsA(UArrayProperty::StaticClass()) || ((Outer != nullptr) && Outer->IsA(UArrayProperty::StaticClass()));
+
+		if (SelectedSection->GetSettingsObject()->GetClass()->HasAnyClassFlags(CLASS_DefaultConfig) && !bIsArrayOrArrayElement)
 		{
 			SelectedSection->GetSettingsObject()->UpdateSinglePropertyInConfigFile(PropertyThatChanged->GetActiveMemberNode()->GetValue(), SelectedSection->GetSettingsObject()->GetDefaultConfigFilename());
 		}
