UE-20019 - fixed issue with array element ini setting not saving out correctly when changed

diff --git a/Engine/Source/Developer/SettingsEditor/Private/Widgets/SSettingsEditor.cpp b/Engine/Source/Developer/SettingsEditor/Private/Widgets/SSettingsEditor.cpp
--- a/Engine/Source/Developer/SettingsEditor/Private/Widgets/SSettingsEditor.cpp
+++ b/Engine/Source/Developer/SettingsEditor/Private/Widgets/SSettingsEditor.cpp
@@ -246,7 +246,7 @@ void SSettingsEditor::NotifyPostChange( const FPropertyChangedEvent& PropertyCha
 
 		// Determine if the Property is an Array or Array Element
 		UObject* Outer = PropertyChangedEvent.Property->GetOuter();
-		bool bIsArrayOrArrayElement = PropertyChangedEvent.Property->IsA(UArrayProperty::StaticClass()) || ((Outer != nullptr) && Outer->IsA(UArrayProperty::StaticClass()));
+		bool bIsArrayOrArrayElement = PropertyThatChanged->GetActiveMemberNode()->GetValue()->IsA(UArrayProperty::StaticClass()) || ((Outer != nullptr) && Outer->IsA(UArrayProperty::StaticClass()));
 
 		if (SelectedSection->GetSettingsObject()->GetClass()->HasAnyClassFlags(CLASS_DefaultConfig) && !bIsArrayOrArrayElement)
 		{
