Copying //UE4/Dev-Editor to //UE4/Dev-Main (Source: //UE4/Dev-Editor @ 3082391)

diff --git a/Engine/Source/Developer/SettingsEditor/Private/Widgets/SSettingsEditor.cpp b/Engine/Source/Developer/SettingsEditor/Private/Widgets/SSettingsEditor.cpp
--- a/Engine/Source/Developer/SettingsEditor/Private/Widgets/SSettingsEditor.cpp
+++ b/Engine/Source/Developer/SettingsEditor/Private/Widgets/SSettingsEditor.cpp
@@ -2,19 +2,15 @@
 
 #include "SettingsEditorPrivatePCH.h"
 #include "AnalyticsEventAttribute.h"
-#include "DesktopPlatformModule.h"
 #include "EngineAnalytics.h"
 #include "IAnalyticsProvider.h"
-#include "NotificationManager.h"
 #include "PropertyEditing.h"
 #include "SHyperlink.h"
-#include "SNotificationList.h"
+#include "SSettingsSectionHeader.h"
 #include "SSettingsEditorCheckoutNotice.h"
 
-
 #define LOCTEXT_NAMESPACE "SSettingsEditor"
 
-
 /* SSettingsEditor structors
  *****************************************************************************/
 
@@ -47,183 +43,85 @@ void SSettingsEditor::Construct( const FArguments& InArgs, const ISettingsEditor
 		DetailsViewArgs.NotifyHook = this;
 		DetailsViewArgs.bShowOptions = true;
 		DetailsViewArgs.bShowModifiedPropertiesOption = false;
+		DetailsViewArgs.bAllowMultipleTopLevelObjects = true;
+		DetailsViewArgs.bShowActorLabel = false;
+		DetailsViewArgs.bCustomNameAreaLocation = true;
+		DetailsViewArgs.bCustomFilterAreaLocation = true;
+		DetailsViewArgs.NameAreaSettings = FDetailsViewArgs::HideNameArea;
 	}
 
 	SettingsView = FModuleManager::GetModuleChecked<FPropertyEditorModule>("PropertyEditor").CreateDetailView(DetailsViewArgs);
 	SettingsView->SetVisibility(TAttribute<EVisibility>::Create(TAttribute<EVisibility>::FGetter::CreateSP(this, &SSettingsEditor::HandleSettingsViewVisibility)));
 	SettingsView->SetIsPropertyEditingEnabledDelegate(FIsPropertyEditingEnabled::CreateSP(this, &SSettingsEditor::HandleSettingsViewEnabled));
 
-	// Create the watcher widget for the default config file (checks file status / SCC state)
-	FileWatcherWidget = SNew(SSettingsEditorCheckoutNotice)
-		.Visibility(this, &SSettingsEditor::HandleCheckoutNoticeVisibility)
-		.OnFileProbablyModifiedExternally(this, &SSettingsEditor::HandleCheckoutNoticeFileProbablyModifiedExternally)
-		.ConfigFilePath(this, &SSettingsEditor::HandleCheckoutNoticeConfigFilePath);
+	TSharedPtr<FSettingsDetailRootObjectCustomization> RootObjectCustomization = MakeShareable(new FSettingsDetailRootObjectCustomization(Model, SettingsView.ToSharedRef()));
+	RootObjectCustomization->Initialize();
+	SettingsView->SetRootObjectCustomizationInstance(RootObjectCustomization);
 
 	ChildSlot
 	[
 		SNew(SVerticalBox)
 
 		+ SVerticalBox::Slot()
-			.FillHeight(1.0f)
-			.Padding(16.0f, 0.0f)
+		.FillHeight(1.0f)
+		.Padding(16.0f, 0.0f)
+		[
+			SNew(SHorizontalBox)
+			+ SHorizontalBox::Slot()
+			.AutoWidth()
+			.Padding(0.0f, 16.0f)
 			[
-				SNew(SHorizontalBox)
-
-				+ SHorizontalBox::Slot()
+				// categories menu
+				SNew(SScrollBox)
+				+ SScrollBox::Slot()
+				[
+					SNew(SHorizontalBox)
+					+ SHorizontalBox::Slot()
 					.AutoWidth()
-					.Padding(0.0f, 16.0f)
 					[
-						// categories menu
-						SNew(SScrollBox)
-
-						+ SScrollBox::Slot()
-							[
-								SNew(SHorizontalBox)
-
-								+ SHorizontalBox::Slot()
-									.AutoWidth()
-									[
-										SAssignNew(CategoriesBox, SVerticalBox)
-									]
-
-								+ SHorizontalBox::Slot()
-									.AutoWidth()
-									[
-										SNew(SSpacer)
-											.Size(FVector2D(24.0f, 0.0f))
-									]
-							]
+						SAssignNew(CategoriesBox, SVerticalBox)
 					]
 
-				+ SHorizontalBox::Slot()
+					+ SHorizontalBox::Slot()
 					.AutoWidth()
-					.Padding(FMargin(24.0f, 0.0f, 24.0f, 0.0f))
 					[
-						SNew(SSeparator)
-							.Orientation(Orient_Vertical)
+						SNew(SSpacer)
+						.Size(FVector2D(24.0f, 0.0f))
 					]
-
-				+ SHorizontalBox::Slot()
-					.FillWidth(1.0f)
-					.Padding(0.0f, 16.0f)
+				]
+			]
+			+ SHorizontalBox::Slot()
+			.AutoWidth()
+			.Padding(FMargin(24.0f, 0.0f, 24.0f, 0.0f))
+			[
+				SNew(SSeparator)
+				.Orientation(Orient_Vertical)
+			]
+			+ SHorizontalBox::Slot()
+			.FillWidth(1.0f)
+			.Padding(0.0f, 16.0f)
+			[
+				SNew(SVerticalBox)
+				.Visibility(this, &SSettingsEditor::HandleSettingsBoxVisibility)
+				+ SVerticalBox::Slot()
+				.AutoHeight()
+				[
+					SettingsView->GetFilterAreaWidget().ToSharedRef()
+				]
+				+ SVerticalBox::Slot()
+				.FillHeight(1.0f)
+				[
+					// settings area
+					SNew(SOverlay)
+					+ SOverlay::Slot()
 					[
-						SNew(SVerticalBox)
-							.Visibility(this, &SSettingsEditor::HandleSettingsBoxVisibility)
-
-						+ SVerticalBox::Slot()
-							.AutoHeight()
-							[
-								// title and button bar
-								SNew(SHorizontalBox)
-
-								+ SHorizontalBox::Slot()
-									.FillWidth(1.0f)
-									[
-										SNew(SVerticalBox)
-
-										+ SVerticalBox::Slot()
-											.AutoHeight()
-											[
-												// category title
-												SNew(STextBlock)
-													.Font(FSlateFontInfo(FPaths::EngineContentDir() / TEXT("Slate/Fonts/Roboto-Regular.ttf"), 18))
-													.Text(this, &SSettingsEditor::HandleSettingsBoxTitleText)													
-											]
-
-										+ SVerticalBox::Slot()
-											.AutoHeight()
-											.Padding(0.0f, 8.0f, 0.0f, 0.0f)
-											[
-												// category description
-												SNew(STextBlock)
-													.ColorAndOpacity(FSlateColor::UseSubduedForeground())
-													.Text(this, &SSettingsEditor::HandleSettingsBoxDescriptionText)
-											]
-									]
-
-								+ SHorizontalBox::Slot()
-									.AutoWidth()
-									.HAlign(HAlign_Right)
-									.VAlign(VAlign_Bottom)
-									.Padding(16.0f, 0.0f, 0.0f, 0.0f)
-									[
-										// set as default button
-										SNew(SButton)
-											.Visibility(this, &SSettingsEditor::HandleSetAsDefaultButtonVisibility)
-											.IsEnabled(this, &SSettingsEditor::HandleSetAsDefaultButtonEnabled)
-											.OnClicked(this, &SSettingsEditor::HandleSetAsDefaultButtonClicked)
-											.Text(LOCTEXT("SaveDefaultsButtonText", "Set as Default"))
-											.ToolTipText(LOCTEXT("SaveDefaultsButtonTooltip", "Save the values below as the new default settings"))
-									]
-
-								+ SHorizontalBox::Slot()
-									.AutoWidth()
-									.HAlign(HAlign_Right)
-									.VAlign(VAlign_Bottom)
-									.Padding(8.0f, 0.0f, 0.0f, 0.0f)
-									[
-										// export button
-										SNew(SButton)
-											.IsEnabled(this, &SSettingsEditor::HandleExportButtonEnabled)
-											.OnClicked(this, &SSettingsEditor::HandleExportButtonClicked)
-											.Text(LOCTEXT("ExportButtonText", "Export..."))
-											.ToolTipText(LOCTEXT("ExportButtonTooltip", "Export these settings to a file on your computer"))
-									]
-
-								+ SHorizontalBox::Slot()
-									.AutoWidth()
-									.HAlign(HAlign_Right)
-									.VAlign(VAlign_Bottom)
-									.Padding(8.0f, 0.0f, 0.0f, 0.0f)
-									[
-										// import button
-										SNew(SButton)
-											.IsEnabled(this, &SSettingsEditor::HandleImportButtonEnabled)
-											.OnClicked(this, &SSettingsEditor::HandleImportButtonClicked)
-											.Text(LOCTEXT("ImportButtonText", "Import..."))
-											.ToolTipText(LOCTEXT("ImportButtonTooltip", "Import these settings from a file on your computer"))
-									]
-
-								+ SHorizontalBox::Slot()
-									.AutoWidth()
-									.HAlign(HAlign_Right)
-									.VAlign(VAlign_Bottom)
-									.Padding(8.0f, 0.0f, 0.0f, 0.0f)
-									[
-										// reset defaults button
-										SNew(SButton)
-											.Visibility(this, &SSettingsEditor::HandleSetAsDefaultButtonVisibility)
-											.IsEnabled(this, &SSettingsEditor::HandleResetToDefaultsButtonEnabled)
-											.OnClicked(this, &SSettingsEditor::HandleResetDefaultsButtonClicked)
-											.Text(LOCTEXT("ResetDefaultsButtonText", "Reset to Defaults"))
-											.ToolTipText(LOCTEXT("ResetDefaultsButtonTooltip", "Reset the settings below to their default values"))
-									]
-							]
-
-						+ SVerticalBox::Slot()
-							.AutoHeight()
-							.Padding(0.0f, 16.0f, 0.0f, 0.0f)
-							[
-								FileWatcherWidget.ToSharedRef()
-							]
-
-						+ SVerticalBox::Slot()
-							.FillHeight(1.0f)
-							.Padding(0.0f, 16.0f, 0.0f, 0.0f)
-							[
-								// settings area
-								SNew(SOverlay)
-
-								+ SOverlay::Slot()
-									[
-										SettingsView.ToSharedRef()
-									]
-
-								+ SOverlay::Slot()
-									.Expose(CustomWidgetSlot)
-							]
+						SettingsView.ToSharedRef()
 					]
+					+ SOverlay::Slot()
+					.Expose(CustomWidgetSlot)
+				]
 			]
+		]
 	];
 
 	FInternationalization::Get().OnCultureChanged().AddSP(this, &SSettingsEditor::HandleCultureChanged);
@@ -238,106 +136,54 @@ void SSettingsEditor::Construct( const FArguments& InArgs, const ISettingsEditor
 
 void SSettingsEditor::NotifyPostChange( const FPropertyChangedEvent& PropertyChangedEvent, class FEditPropertyChain* PropertyThatChanged )
 {
-	ISettingsSectionPtr SelectedSection = Model->GetSelectedSection();
-
-	if (SelectedSection.IsValid() && (PropertyChangedEvent.ChangeType != EPropertyChangeType::Interactive))
+	if (PropertyChangedEvent.ChangeType != EPropertyChangeType::Interactive)
 	{
-		RecordPreferenceChangedAnalytics( SelectedSection, PropertyChangedEvent );
-
-		// Determine if the Property is an Array or Array Element
 		UObject* Outer = PropertyChangedEvent.Property->GetOuter();
 
-		bool bIsArrayOrArrayElement = PropertyThatChanged->GetActiveMemberNode()->GetValue()->IsA(UArrayProperty::StaticClass()) 
-			|| PropertyThatChanged->GetActiveMemberNode()->GetValue()->ArrayDim > 1
-			|| ((Outer != nullptr) && Outer->IsA(UArrayProperty::StaticClass()));
+		// Note while there could be multiple objects in the details panel, only one is ever edited at once
+		const UObject* ObjectBeingEdited = PropertyChangedEvent.GetObjectBeingEdited(0);
 
-		if (SelectedSection->GetSettingsObject()->GetClass()->HasAnyClassFlags(CLASS_DefaultConfig) && !bIsArrayOrArrayElement)
-		{
-			SelectedSection->GetSettingsObject()->UpdateSinglePropertyInConfigFile(PropertyThatChanged->GetActiveMemberNode()->GetValue(), SelectedSection->GetSettingsObject()->GetDefaultConfigFilename());
-		}
-		else
-		{
-			SelectedSection->Save();
-		}
 
-		static const FName ConfigRestartRequiredKey = "ConfigRestartRequired";
-		if (PropertyChangedEvent.Property->GetBoolMetaData(ConfigRestartRequiredKey))
+		// Get the section from the edited object.  We cannot use the selected section as multiple sections can be shown at once in the settings details panel.
+		ISettingsSectionPtr Section = Model->GetSectionFromSectionObject(ObjectBeingEdited);
+		if(Section.IsValid())
 		{
-			OnApplicationRestartRequiredDelegate.ExecuteIfBound();
-		}
-	}
-}
-
-
-/* SSettingsEditor implementation
- *****************************************************************************/
-
-bool SSettingsEditor::CheckOutDefaultConfigFile()
-{
-	TWeakObjectPtr<UObject> SettingsObject = GetSelectedSettingsObject();
-
-	if (!SettingsObject.IsValid())
-	{
-		return false;
-	}
-
-	FText ErrorMessage;
+			// Attempt to checkout the file automatically
+			FString FullPath = FPaths::ConvertRelativePathToFull(ObjectBeingEdited->GetDefaultConfigFilename());
 
-	// check out configuration file
-	FString AbsoluteConfigFilePath = GetDefaultConfigFilePath(SettingsObject);
-	TArray<FString> FilesToBeCheckedOut;
-	FilesToBeCheckedOut.Add(AbsoluteConfigFilePath);
+			if(!SettingsHelpers::CheckOutFile(FullPath))
+			{
+				SettingsHelpers::MakeWritable(FullPath);
+			}
+			
+			RecordPreferenceChangedAnalytics(Section, PropertyChangedEvent);
 
-	ISourceControlProvider& SourceControlProvider = ISourceControlModule::Get().GetProvider();
-	FSourceControlStatePtr SourceControlState = SourceControlProvider.GetState(AbsoluteConfigFilePath, EStateCacheUsage::ForceUpdate);
+			// Determine if the Property is an Array or Array Element
+			bool bIsArrayOrArrayElement = PropertyThatChanged->GetActiveMemberNode()->GetValue()->IsA(UArrayProperty::StaticClass()) 
+				|| PropertyThatChanged->GetActiveMemberNode()->GetValue()->ArrayDim > 1
+				|| ((Outer != nullptr) && Outer->IsA(UArrayProperty::StaticClass()));
 
-	if (SourceControlState.IsValid())
-	{
-		if (SourceControlState->IsDeleted())
-		{
-			ErrorMessage = LOCTEXT("ConfigFileMarkedForDeleteError", "Error: The configuration file is marked for deletion.");
-		}
-		else if (!SourceControlState->IsCurrent())
-		{
-			if (false)
+			if (Section->GetSettingsObject()->GetClass()->HasAnyClassFlags(CLASS_DefaultConfig) && !bIsArrayOrArrayElement)
 			{
-				if (SourceControlProvider.Execute(ISourceControlOperation::Create<FSync>(), FilesToBeCheckedOut) == ECommandResult::Succeeded)
-				{
-					SettingsObject->ReloadConfig();
-				}
-				else
-				{
-					ErrorMessage = LOCTEXT("FailedToSyncConfigFileError", "Error: Failed to sync the configuration file to head revision.");
-				}
+				Section->GetSettingsObject()->UpdateSinglePropertyInConfigFile(PropertyThatChanged->GetActiveMemberNode()->GetValue(), Section->GetSettingsObject()->GetDefaultConfigFilename());
 			}
-		}
-		else if (SourceControlState->CanCheckout() || SourceControlState->IsCheckedOutOther())
-		{
-			if (SourceControlProvider.Execute(ISourceControlOperation::Create<FCheckOut>(), FilesToBeCheckedOut) == ECommandResult::Failed)
+			else
 			{
-				ErrorMessage = LOCTEXT("FailedToCheckOutConfigFileError", "Error: Failed to check out the configuration file.");
+				Section->Save();
 			}
-		}
-	}
 
-	// show errors, if any
-	if (!ErrorMessage.IsEmpty())
-	{
-		FMessageDialog::Open(EAppMsgType::Ok, ErrorMessage);
-
-		return false;
+			static const FName ConfigRestartRequiredKey = "ConfigRestartRequired";
+			if (PropertyChangedEvent.Property->GetBoolMetaData(ConfigRestartRequiredKey))
+			{
+				OnApplicationRestartRequiredDelegate.ExecuteIfBound();
+			}
+		}
 	}
-
-	return true;
 }
 
 
-FString SSettingsEditor::GetDefaultConfigFilePath( const TWeakObjectPtr<UObject>& SettingsObject ) const
-{
-	FString RelativeConfigFilePath = SettingsObject->GetDefaultConfigFilename();
-	return FPaths::ConvertRelativePathToFull(RelativeConfigFilePath);
-}
-
+/* SSettingsEditor implementation
+ *****************************************************************************/
 
 TWeakObjectPtr<UObject> SSettingsEditor::GetSelectedSettingsObject() const
 {
@@ -351,56 +197,19 @@ TWeakObjectPtr<UObject> SSettingsEditor::GetSelectedSettingsObject() const
 	return nullptr;
 }
 
-
-bool SSettingsEditor::IsDefaultConfigCheckOutNeeded() const
-{
-	TWeakObjectPtr<UObject> SettingsObject = GetSelectedSettingsObject();
-
-	if (SettingsObject.IsValid() && SettingsObject->GetClass()->HasAnyClassFlags(CLASS_Config | CLASS_DefaultConfig))
-	{
-		return !FileWatcherWidget->IsUnlocked();
-	}
-	else
-	{
-		return false;
-	}
-}
-
-
-TSharedRef<SWidget> SSettingsEditor::MakeCategoryWidget( const ISettingsCategoryRef& Category )
+TSharedRef<SWidget> SSettingsEditor::MakeCategoryWidget( const ISettingsCategoryRef& Category, const TArray<ISettingsSectionPtr>& InSections )
 {
 	// create section widgets
 	TSharedRef<SVerticalBox> SectionsBox = SNew(SVerticalBox);
-	TArray<ISettingsSectionPtr> Sections;
+
 	
-	if (Category->GetSections(Sections) == 0)
+	if (InSections.Num() == 0)
 	{
 		return SNullWidget::NullWidget;
 	}
 
-	// sort the sections alphabetically
-	struct FSectionSortPredicate
-	{
-		FORCEINLINE bool operator()( ISettingsSectionPtr A, ISettingsSectionPtr B ) const
-		{
-			if (!A.IsValid() && !B.IsValid())
-			{
-				return false;
-			}
-
-			if (A.IsValid() != B.IsValid())
-			{
-				return B.IsValid();
-			}
-
-			return (A->GetDisplayName().CompareTo(B->GetDisplayName()) < 0);
-		}
-	};
-
-	Sections.Sort(FSectionSortPredicate());
-
 	// list the sections
-	for (const ISettingsSectionPtr Section : Sections)
+	for (const ISettingsSectionPtr Section : InSections)
 	{
 		SectionsBox->AddSlot()
 			.HAlign(HAlign_Left)
@@ -443,7 +252,7 @@ TSharedRef<SWidget> SSettingsEditor::MakeCategoryWidget( const ISettingsCategory
 		[
 			// category title
 			SNew(STextBlock)
-				.Font(FSlateFontInfo(FPaths::EngineContentDir() / TEXT("Slate/Fonts/Roboto-Regular.ttf"), 18))
+				.Font(FEditorStyle::GetFontStyle("SettingsEditor.CatgoryAndSectionFont"))
 				.Text(Category->GetDisplayName())
 		]
 
@@ -456,21 +265,6 @@ TSharedRef<SWidget> SSettingsEditor::MakeCategoryWidget( const ISettingsCategory
 }
 
 
-bool SSettingsEditor::MakeDefaultConfigFileWritable()
-{
-	TWeakObjectPtr<UObject> SettingsObject = GetSelectedSettingsObject();
-
-	if (!SettingsObject.IsValid())
-	{
-		return false;
-	}
-
-	FString AbsoluteConfigFilePath = GetDefaultConfigFilePath(SettingsObject);
-
-	return FPlatformFileManager::Get().GetPlatformFile().SetReadOnly(*AbsoluteConfigFilePath, false);
-}
-
-
 void SSettingsEditor::RecordPreferenceChangedAnalytics( ISettingsSectionPtr SelectedSection, const FPropertyChangedEvent& PropertyChangedEvent ) const
 {
 	UProperty* ChangedProperty = PropertyChangedEvent.MemberProperty;
@@ -491,167 +285,99 @@ void SSettingsEditor::ReloadCategories()
 {
 	CategoriesBox->ClearChildren();
 
+	CategoriesBox->AddSlot()
+	.AutoHeight()
+	.HAlign(HAlign_Left)
+	.Padding(0.0f, 5.0f, 0.0f, 18.0f)
+	[
+		SNew(SHorizontalBox)
+		+ SHorizontalBox::Slot()
+		.AutoWidth()
+		.Padding(0.0f, 0.0f, 4.0f, 0.0f)
+		.VAlign(VAlign_Center)
+		[
+			SNew(SImage)
+			.Image(FEditorStyle::Get().GetBrush("TreeArrow_Collapsed_Hovered"))
+			.Visibility(this, &SSettingsEditor::HandleAllSectionsLinkImageVisibility)
+		]
+		+ SHorizontalBox::Slot()
+		.FillWidth(1.0f)
+		.VAlign(VAlign_Center)
+		[
+			SNew(SHyperlink)
+			.OnNavigate(this, &SSettingsEditor::HandleAllSectionsLinkNavigate)
+			.Text(LOCTEXT("AllPropertiesLink", "All Settings"))
+			.ToolTipText(LOCTEXT("AllPropertiesLink_Tooltip", "Show all settings"))
+		]
+	];
+
 	TArray<ISettingsCategoryPtr> Categories;
 	SettingsContainer->GetCategories(Categories);
 
-	for (const ISettingsCategoryPtr Category : Categories)
-	{
-		CategoriesBox->AddSlot()
-			.AutoHeight()
-			.Padding(0.0f, 0.0f, 0.0f, 16.0f)
-			[
-				MakeCategoryWidget(Category.ToSharedRef())
-			];
-	}
-}
+	TArray<TSharedPtr<ISettingsSection>> CategorySettingsSections;
 
-
-void SSettingsEditor::ShowNotification( const FText& Text, SNotificationItem::ECompletionState CompletionState ) const
-{
-	FNotificationInfo Notification(Text);
-	Notification.ExpireDuration = 3.f;
-	Notification.bUseSuccessFailIcons = false;
-
-	FSlateNotificationManager::Get().AddNotification(Notification)->SetCompletionState(CompletionState);
-}
-
-
-/* SSettingsEditor callbacks
- *****************************************************************************/
-
-FString SSettingsEditor::HandleCheckoutNoticeConfigFilePath() const
-{
-	TWeakObjectPtr<UObject> SettingsObject = GetSelectedSettingsObject();
-
-	if (SettingsObject.IsValid() && SettingsObject->GetClass()->HasAnyClassFlags(CLASS_Config | CLASS_DefaultConfig))
-	{
-		return GetDefaultConfigFilePath(SettingsObject);
-	}
-	else
-	{
-		return FString();
-	}
-}
-
-
-void SSettingsEditor::HandleCheckoutNoticeFileProbablyModifiedExternally()
-{
-	TWeakObjectPtr<UObject> SettingsObject = GetSelectedSettingsObject();
-
-	if (SettingsObject.IsValid() && SettingsObject->GetClass()->HasAnyClassFlags(CLASS_Config | CLASS_DefaultConfig))
-	{
-		SettingsObject->ReloadConfig();
-	}
-}
-
-
-EVisibility SSettingsEditor::HandleCheckoutNoticeVisibility() const
-{
-	TWeakObjectPtr<UObject> SettingsObject = GetSelectedSettingsObject();
-
-	if (SettingsObject.IsValid() && SettingsObject->GetClass()->HasAnyClassFlags(CLASS_DefaultConfig))
+	TArray<UObject*> SettingsObjects;
+	for (const ISettingsCategoryPtr Category : Categories)
 	{
-		return EVisibility::Visible;
-	}
+		CategorySettingsSections.Reset();
 
-	return EVisibility::Collapsed;	
-}
-
-
-void SSettingsEditor::HandleCultureChanged()
-{
-	ReloadCategories();
-}
+		Category->GetSections(CategorySettingsSections);
 
+		// sort the sections alphabetically
+		struct FSectionSortPredicate
+		{
+			FORCEINLINE bool operator()(ISettingsSectionPtr A, ISettingsSectionPtr B) const
+			{
+				if (!A.IsValid() && !B.IsValid())
+				{
+					return false;
+				}
 
-FReply SSettingsEditor::HandleExportButtonClicked()
-{
-	ISettingsSectionPtr SelectedSection = Model->GetSelectedSection();
+				if (A.IsValid() != B.IsValid())
+				{
+					return B.IsValid();
+				}
 
-	if (SelectedSection.IsValid())
-	{
-		if (LastExportDir.IsEmpty())
-		{
-			LastExportDir = FPaths::GetPath(GEditorPerProjectIni);
-		}
+				return (A->GetDisplayName().CompareTo(B->GetDisplayName()) < 0);
+			}
+		};
 
-		FString DefaultFileName = FString::Printf(TEXT("%s Backup %s.ini"), *SelectedSection->GetDisplayName().ToString(), *FDateTime::Now().ToString(TEXT("%Y-%m-%d %H%M%S")));
-		TArray<FString> OutFiles;
+		CategorySettingsSections.Sort(FSectionSortPredicate());
 
-		TSharedPtr<SWindow> ParentWindow = FSlateApplication::Get().FindWidgetWindow(AsShared());
-		void* ParentWindowHandle = (ParentWindow.IsValid() && ParentWindow->GetNativeWindow().IsValid()) ? ParentWindow->GetNativeWindow()->GetOSWindowHandle() : nullptr;
+		CategoriesBox->AddSlot()
+			.AutoHeight()
+			.Padding(0.0f, 0.0f, 0.0f, 16.0f)
+			[
+				MakeCategoryWidget(Category.ToSharedRef(), CategorySettingsSections)
+			];
 
-		if (FDesktopPlatformModule::Get()->SaveFileDialog(ParentWindowHandle, LOCTEXT("ExportSettingsDialogTitle", "Export settings...").ToString(), LastExportDir, DefaultFileName, TEXT("Config files (*.ini)|*.ini"), EFileDialogFlags::None, OutFiles))
+		for(TSharedPtr<ISettingsSection>& Section : CategorySettingsSections)
 		{
-			if (SelectedSection->Export(OutFiles[0]))
+			UObject* SettingsObject = Section->GetSettingsObject().Get();
+			if(SettingsObject)
 			{
-				ShowNotification(LOCTEXT("ExportSettingsSuccess", "Export settings succeeded"), SNotificationItem::CS_Success);
-			}
-			else
-			{
-				ShowNotification(LOCTEXT("ExportSettingsFailure", "Export settings failed"), SNotificationItem::CS_Fail);
+				SettingsObjects.Add(SettingsObject);
 			}
 		}
 	}
 
-	return FReply::Handled();
-}
-
-
-bool SSettingsEditor::HandleExportButtonEnabled() const
-{
-	ISettingsSectionPtr SelectedSection = Model->GetSelectedSection();
+	SettingsView->SetObjects(SettingsObjects);
 
-	if (SelectedSection.IsValid())
-	{
-		return SelectedSection->CanExport();
-	}
 
-	return false;
 }
 
 
-FReply SSettingsEditor::HandleImportButtonClicked()
-{
-	ISettingsSectionPtr SelectedSection = Model->GetSelectedSection();
 
-	if (SelectedSection.IsValid())
-	{
-		TArray<FString> OutFiles;
 
-		TSharedPtr<SWindow> ParentWindow = FSlateApplication::Get().FindWidgetWindow(AsShared());
-		void* ParentWindowHandle = (ParentWindow.IsValid() && ParentWindow->GetNativeWindow().IsValid()) ? ParentWindow->GetNativeWindow()->GetOSWindowHandle() : nullptr;
-
-		if (FDesktopPlatformModule::Get()->OpenFileDialog(ParentWindowHandle, LOCTEXT("ImportSettingsDialogTitle", "Import settings...").ToString(), FPaths::GetPath(GEditorPerProjectIni), TEXT(""), TEXT("Config files (*.ini)|*.ini"), EFileDialogFlags::None, OutFiles))
-		{
-			if (SelectedSection->Import(OutFiles[0]) && SelectedSection->Save())
-			{
-				ShowNotification(LOCTEXT("ImportSettingsSuccess", "Import settings succeeded"), SNotificationItem::CS_Success);
-			}
-			else
-			{
-				ShowNotification(LOCTEXT("ImportSettingsFailure", "Import settings failed"), SNotificationItem::CS_Fail);
-			}
-		}
-	}
-
-	return FReply::Handled();
-}
+/* SSettingsEditor callbacks
+ *****************************************************************************/
 
 
-bool SSettingsEditor::HandleImportButtonEnabled() const
+void SSettingsEditor::HandleCultureChanged()
 {
-	ISettingsSectionPtr SelectedSection = Model->GetSelectedSection();
-
-	if (SelectedSection.IsValid())
-	{
-		return SelectedSection->CanEdit() && SelectedSection->CanImport() && !IsDefaultConfigCheckOutNeeded();
-	}
-
-	return false;
+	ReloadCategories();
 }
 
-
 void SSettingsEditor::HandleModelSelectionChanged()
 {
 	ISettingsSectionPtr SelectedSection = Model->GetSelectedSection();
@@ -670,8 +396,6 @@ void SSettingsEditor::HandleModelSelectionChanged()
 			CustomWidgetSlot->AttachWidget( SNullWidget::NullWidget );
 		}
 
-		SettingsView->SetObject(SelectedSection->GetSettingsObject().Get());
-
 		// focus settings widget
 		TSharedPtr<SWidget> FocusWidget;
 
@@ -690,57 +414,35 @@ void SSettingsEditor::HandleModelSelectionChanged()
 		{
 			FSlateApplication::Get().SetKeyboardFocus(FocusWidgetPath, EFocusCause::SetDirectly);
 		}
+
+		bShowingAllSettings = false;
 	}
 	else
 	{
+		bShowingAllSettings = true;
 		CustomWidgetSlot->AttachWidget( SNullWidget::NullWidget );
-		SettingsView->SetObject(nullptr);
-	}
-
-	FileWatcherWidget->Invalidate();
-}
-
-
-FReply SSettingsEditor::HandleResetDefaultsButtonClicked()
-{
-	ISettingsSectionPtr SelectedSection = Model->GetSelectedSection();
-
-	if (SelectedSection.IsValid())
-	{
-		SelectedSection->ResetDefaults();
 	}
 
-	return FReply::Handled();
+	// clear the global search terms when selecting a specific category
+	SettingsView->ClearSearch();
 }
 
-
-bool SSettingsEditor::HandleResetToDefaultsButtonEnabled() const
+void SSettingsEditor::HandleSectionLinkNavigate( ISettingsSectionPtr Section )
 {
-	ISettingsSectionPtr SelectedSection = Model->GetSelectedSection();
-
-	if (SelectedSection.IsValid())
-	{
-		return (SelectedSection->CanEdit() && SelectedSection->CanResetDefaults());
-	}
-
-	return false;
+	Model->SelectSection(Section);
 }
 
-
-EVisibility SSettingsEditor::HandleSetAsDefaultButtonVisibility() const
+void SSettingsEditor::HandleAllSectionsLinkNavigate()
 {
-	ISettingsSectionPtr SelectedSection = Model->GetSelectedSection();
-
-	return (SelectedSection.IsValid() && SelectedSection->HasDefaultSettingsObject()) ? EVisibility::Collapsed : EVisibility::Visible;
+	Model->SelectSection(nullptr);
+	SettingsView->RefreshRootObjectVisibility();
 }
 
-
-void SSettingsEditor::HandleSectionLinkNavigate( ISettingsSectionPtr Section )
+EVisibility SSettingsEditor::HandleAllSectionsLinkImageVisibility() const
 {
-	Model->SelectSection(Section);
+	return bShowingAllSettings ? EVisibility::Visible : EVisibility::Hidden;
 }
 
-
 EVisibility SSettingsEditor::HandleSectionLinkImageVisibility( ISettingsSectionPtr Section ) const
 {
 	if (Model->GetSelectedSection() == Section)
@@ -751,94 +453,11 @@ EVisibility SSettingsEditor::HandleSectionLinkImageVisibility( ISettingsSectionP
 	return EVisibility::Hidden;
 }
 
-
-FReply SSettingsEditor::HandleSetAsDefaultButtonClicked()
-{
-	ISettingsSectionPtr SelectedSection = Model->GetSelectedSection();
-
-	if (SelectedSection.IsValid())
-	{
-		if (FMessageDialog::Open(EAppMsgType::YesNo, LOCTEXT("SaveAsDefaultUserConfirm", "Are you sure you want to update the default settings?")) != EAppReturnType::Yes)
-		{
-			return FReply::Handled();
-		}
-
-		if (IsDefaultConfigCheckOutNeeded())
-		{
-			if (ISourceControlModule::Get().IsEnabled())
-			{
-				if (FMessageDialog::Open(EAppMsgType::YesNo, LOCTEXT("SaveAsDefaultNeedsCheckoutMessage", "The default configuration file for these settings is currently not checked out. Would you like to check it out from source control? \n\nIf checked into source control this would affect other developers.")) != EAppReturnType::Yes)
-				{
-					return FReply::Handled();
-				}
-
-				CheckOutDefaultConfigFile();
-			}
-			else
-			{
-				if (FMessageDialog::Open(EAppMsgType::YesNo, LOCTEXT("SaveAsDefaultsIsReadOnlyMessage", "The default configuration file for these settings is currently not writeable. Would you like to make it writable?")) != EAppReturnType::Yes)
-				{
-					return FReply::Handled();
-				}
-
-				MakeDefaultConfigFileWritable();
-			}
-		}
-
-		SelectedSection->SaveDefaults();
-
-		FMessageDialog::Open(EAppMsgType::Ok, LOCTEXT("SaveAsDefaultsSucceededMessage", "The default configuration file for these settings was updated successfully. \n\nIf checked into source control this would affect other developers."));
-	}
-
-	return FReply::Handled();
-}
-
-
-bool SSettingsEditor::HandleSetAsDefaultButtonEnabled() const
-{
-	ISettingsSectionPtr SelectedSection = Model->GetSelectedSection();
-
-	if (SelectedSection.IsValid())
-	{
-		return SelectedSection->CanSaveDefaults();
-	}
-
-	return false;
-}
-
-
-FText SSettingsEditor::HandleSettingsBoxDescriptionText() const
-{
-	ISettingsSectionPtr SelectedSection = Model->GetSelectedSection();
-
-	if (SelectedSection.IsValid())
-	{
-		return SelectedSection->GetDescription();
-	}
-
-	return FText::GetEmpty();
-}
-
-
-FText SSettingsEditor::HandleSettingsBoxTitleText() const
-{
-	ISettingsSectionPtr SelectedSection = Model->GetSelectedSection();
-
-	if (SelectedSection.IsValid())
-	{
-		static const FText TitleFmt = FText::FromString(TEXT("{0} - {1}"));
-		return FText::Format(TitleFmt, SelectedSection->GetCategory().Pin()->GetDisplayName(), SelectedSection->GetDisplayName());
-	}
-
-	return FText::GetEmpty();
-}
-
-
 EVisibility SSettingsEditor::HandleSettingsBoxVisibility() const
 {
 	ISettingsSectionPtr SelectedSection = Model->GetSelectedSection();
 
-	if (SelectedSection.IsValid())
+	if (SelectedSection.IsValid() || bShowingAllSettings)
 	{
 		return EVisibility::Visible;
 	}
@@ -846,7 +465,6 @@ EVisibility SSettingsEditor::HandleSettingsBoxVisibility() const
 	return EVisibility::Hidden;
 }
 
-
 void SSettingsEditor::HandleSettingsContainerCategoryModified( const FName& CategoryName )
 {
 	if ( !bIsActiveTimerRegistered )
@@ -869,7 +487,7 @@ bool SSettingsEditor::HandleSettingsViewEnabled() const
 {
 	ISettingsSectionPtr SelectedSection = Model->GetSelectedSection();
 
-	return (SelectedSection.IsValid() && SelectedSection->CanEdit() && (!SelectedSection->HasDefaultSettingsObject() || FileWatcherWidget->IsUnlocked()));
+	return (SelectedSection.IsValid() && SelectedSection->CanEdit()) || bShowingAllSettings;
 }
 
 
@@ -877,7 +495,7 @@ EVisibility SSettingsEditor::HandleSettingsViewVisibility() const
 {
 	ISettingsSectionPtr SelectedSection = Model->GetSelectedSection();
 
-	if (SelectedSection.IsValid() && SelectedSection->GetSettingsObject().IsValid())
+	if (bShowingAllSettings || (SelectedSection.IsValid() && SelectedSection->GetSettingsObject().IsValid()))
 	{
 		return EVisibility::Visible;
 	}
diff --git a/Engine/Source/Developer/SettingsEditor/Private/Widgets/SSettingsEditor.h b/Engine/Source/Developer/SettingsEditor/Private/Widgets/SSettingsEditor.h
--- a/Engine/Source/Developer/SettingsEditor/Private/Widgets/SSettingsEditor.h
+++ b/Engine/Source/Developer/SettingsEditor/Private/Widgets/SSettingsEditor.h
@@ -50,20 +50,6 @@ public:
 
 protected:
 
-	/**
-	 * Checks out the default configuration file for the currently selected settings object.
-	 *
-	 * @return true if the check-out succeeded, false otherwise.
-	 */
-	bool CheckOutDefaultConfigFile();
-
-	/**
-	 * Gets the absolute path to the Default.ini for the specified object.
-	 *
-	 * @param SettingsObject The object to get the path for.
-	 * @return The path to the file.
-	 */
-	FString GetDefaultConfigFilePath( const TWeakObjectPtr<UObject>& SettingsObject ) const;
 
 	/**
 	 * Gets the settings object of the selected section, if any.
@@ -72,27 +58,14 @@ protected:
 	 */
 	TWeakObjectPtr<UObject> GetSelectedSettingsObject() const;
 
-	/**
-	 * Checks whether the default config file needs to be checked out for editing.
-	 *
-	 * @return true if the file needs to be checked out, false otherwise.
-	 */
-	bool IsDefaultConfigCheckOutNeeded() const;
-
 	/**
 	 * Creates a widget for the given settings category.
 	 *
 	 * @param Category The category to create the widget for.
 	 * @return The created widget.
 	 */
-	TSharedRef<SWidget> MakeCategoryWidget( const ISettingsCategoryRef& Category );
+	TSharedRef<SWidget> MakeCategoryWidget( const ISettingsCategoryRef& Category, const TArray<ISettingsSectionPtr>& InSections );
 
-	/**
-	 * Makes the default configuration file for the currently selected settings object writable.
-	 *
-	 * @return true if it was made writable, false otherwise.
-	 */
-	bool MakeDefaultConfigFileWritable();
 
 	/**
 	 * Reports a preference changed event to the analytics system.
@@ -105,72 +78,26 @@ protected:
 	/** Reloads the settings categories. */
 	void ReloadCategories();
 
-	/**
-	 * Shows a notification pop-up.
-	 *
-	 * @param Text The notification text.
-	 * @param CompletionState The notification's completion state, i.e. success or failure.
-	 */
-	void ShowNotification( const FText& Text, SNotificationItem::ECompletionState CompletionState ) const;
 
 private:
 
-	/** Callback for clicking the Back link. */
-	FReply HandleBackButtonClicked();
-
-	/** Returns the config file name currently being edited. */
-	FString HandleCheckoutNoticeConfigFilePath() const;
-
-	/** Reloads the configuration object. */
-	void HandleCheckoutNoticeFileProbablyModifiedExternally();
-
-	/** Callback for determining the visibility of the 'Locked' notice. */
-	EVisibility HandleCheckoutNoticeVisibility() const;
-
 	/** Callback for when the user's culture has changed. */
 	void HandleCultureChanged();
 
-	/** Callback for clicking the 'Reset to Defaults' button. */
-	FReply HandleExportButtonClicked();
-
-	/** Callback for getting the enabled state of the 'Export' button. */
-	bool HandleExportButtonEnabled() const;
-
-	/** Callback for clicking the 'Reset to Defaults' button. */
-	FReply HandleImportButtonClicked();
-
-	/** Callback for getting the enabled state of the 'Import' button. */
-	bool HandleImportButtonEnabled() const;
-
 	/** Callback for changing the selected settings section. */
 	void HandleModelSelectionChanged();
 
-	/** Callback for clicking the 'Reset to Defaults' button. */
-	FReply HandleResetDefaultsButtonClicked();
-
-	/** Callback for getting the enabled state of the 'Reset to Defaults' button. */
-	bool HandleResetToDefaultsButtonEnabled() const;
-
 	/** Callback for navigating a settings section link. */
 	void HandleSectionLinkNavigate( ISettingsSectionPtr Section );
 
+	/** Callback for navigating the all settings link. */
+	void HandleAllSectionsLinkNavigate();
+
 	/** Callback for getting the visibility of a section link image. */
 	EVisibility HandleSectionLinkImageVisibility( ISettingsSectionPtr Section ) const;
 
-	/** Callback for clicking the 'Set as Defaults' button. */
-	FReply HandleSetAsDefaultButtonClicked();
-
-	/** Callback for getting the enabled state of the 'Set as Defaults' button. */
-	bool HandleSetAsDefaultButtonEnabled() const;
-
-	/** Callback for getting the visibility of options that are only visible when editing something that came from a non-default object. */
-	EVisibility HandleSetAsDefaultButtonVisibility() const;
-
-	/** Callback for getting the section description text. */
-	FText HandleSettingsBoxDescriptionText() const;
-
-	/** Callback for getting the section text for the settings box. */
-	FText HandleSettingsBoxTitleText() const;
+	/** Callback for getting the visibility of a the all sections link image. */
+	EVisibility HandleAllSectionsLinkImageVisibility() const;
 
 	/** Callback for determining the visibility of the settings box. */
 	EVisibility HandleSettingsBoxVisibility() const;
@@ -186,7 +113,6 @@ private:
 
 	/** Callback when the timer has been ticked to refresh the categories latently. */
 	EActiveTimerReturnType UpdateCategoriesCallback(double InCurrentTime, float InDeltaTime);
-
 private:
 
 	/** Holds the vertical box for settings categories. */
@@ -194,13 +120,6 @@ private:
 
 	/** Holds the overlay slot for custom widgets. */
 	SOverlay::FOverlaySlot* CustomWidgetSlot;
-
-	/** Watcher widget for the default config file (checks file status / SCC state). */
-	TSharedPtr<SSettingsEditorCheckoutNotice> FileWatcherWidget;
-
-	/** Holds the path to the directory that the last settings were exported to. */
-	FString LastExportDir;
-
 	/** Holds a pointer to the view model. */
 	ISettingsEditorModelPtr Model;
 
@@ -215,4 +134,7 @@ private:
 
 	/** Is the active timer registered to refresh categories after the settings changed. */
 	bool bIsActiveTimerRegistered;
+	
+	/** Are we showing all settings at once */
+	bool bShowingAllSettings;
 };
diff --git a/Engine/Source/Developer/SettingsEditor/Private/Widgets/SSettingsSectionHeader.cpp b/Engine/Source/Developer/SettingsEditor/Private/Widgets/SSettingsSectionHeader.cpp
new file mode 100644
--- /dev/null
+++ b/Engine/Source/Developer/SettingsEditor/Private/Widgets/SSettingsSectionHeader.cpp
@@ -0,0 +1,506 @@
+// Copyright 1998-2016 Epic Games, Inc. All Rights Reserved.
+
+
+#include "SettingsEditorPrivatePCH.h"
+#include "SSettingsSectionHeader.h"
+#include "NotificationManager.h"
+#include "SNotificationList.h"
+#include "SSettingsEditorCheckoutNotice.h"
+#include "DesktopPlatformModule.h"
+#include "IDetailsView.h"
+
+#define LOCTEXT_NAMESPACE "SSettingsEditor"
+
+void SSettingsSectionHeader::Construct(const FArguments& InArgs, const UObject* InSettingsObject, ISettingsEditorModelPtr InModel, TSharedPtr<IDetailsView> InDetailsView)
+{
+	Model = InModel;
+	SettingsObject = InSettingsObject;
+	SettingsSection = Model->GetSectionFromSectionObject(InSettingsObject);
+	DetailsView = InDetailsView;
+
+	Model->OnSelectionChanged().AddSP(this, &SSettingsSectionHeader::OnSettingsSelectionChanged);
+
+	// Create the watcher widget for the default config file (checks file status / SCC state)
+	FileWatcherWidget =
+		SNew(SSettingsEditorCheckoutNotice)
+		.Visibility(this, &SSettingsSectionHeader::HandleCheckoutNoticeVisibility)
+		.OnFileProbablyModifiedExternally(this, &SSettingsSectionHeader::HandleCheckoutNoticeFileProbablyModifiedExternally)
+		.ConfigFilePath(this, &SSettingsSectionHeader::HandleCheckoutNoticeConfigFilePath);
+
+
+	ChildSlot
+	.Padding(FMargin(0.0f, 8.0f, 0.0f, 5.0f))
+	[
+		SNew(SVerticalBox)
+		+ SVerticalBox::Slot()
+		.AutoHeight()
+		[
+			SNew(SHorizontalBox)
+			+ SHorizontalBox::Slot()
+			.FillWidth(1.0f)
+			[
+				SNew(SVerticalBox)
+				+ SVerticalBox::Slot()
+				.AutoHeight()
+				[
+					// category title
+					SNew(STextBlock)
+					.Font(FEditorStyle::GetFontStyle("SettingsEditor.CatgoryAndSectionFont"))
+					.Text(GetSettingsBoxTitleText())
+				]
+				+ SVerticalBox::Slot()
+				.AutoHeight()
+				.Padding(0.0f, 8.0f, 0.0f, 0.0f)
+				[
+					// category description
+					SNew(STextBlock)
+					.ColorAndOpacity(FSlateColor::UseSubduedForeground())
+					.Text(GetSettingsBoxDescriptionText())
+				]
+			]
+			+ SHorizontalBox::Slot()
+			.AutoWidth()
+			.HAlign(HAlign_Right)
+			.VAlign(VAlign_Bottom)
+			.Padding(16.0f, 0.0f, 0.0f, 0.0f)
+			[
+				SNew(SHorizontalBox)
+				.Visibility(this, &SSettingsSectionHeader::HandleButtonRowVisibility)
+				+SHorizontalBox::Slot()
+				[
+					// set as default button
+					SNew(SButton)
+					.Visibility(this, &SSettingsSectionHeader::HandleSetAsDefaultButtonVisibility)
+					.IsEnabled(this, &SSettingsSectionHeader::HandleSetAsDefaultButtonEnabled)
+					.OnClicked(this, &SSettingsSectionHeader::HandleSetAsDefaultButtonClicked)
+					.Text(LOCTEXT("SaveDefaultsButtonText", "Set as Default"))
+					.ToolTipText(LOCTEXT("SaveDefaultsButtonTooltip", "Save the values below as the new default settings"))
+				]
+				+ SHorizontalBox::Slot()
+				.AutoWidth()
+				.Padding(8.0f, 0.0f, 0.0f, 0.0f)
+				[
+					// export button
+					SNew(SButton)
+					.IsEnabled(this, &SSettingsSectionHeader::HandleExportButtonEnabled)
+					.OnClicked(this, &SSettingsSectionHeader::HandleExportButtonClicked)
+					.Text(LOCTEXT("ExportButtonText", "Export..."))
+					.ToolTipText(LOCTEXT("ExportButtonTooltip", "Export these settings to a file on your computer"))
+				]
+				+ SHorizontalBox::Slot()
+				.AutoWidth()
+				.Padding(8.0f, 0.0f, 0.0f, 0.0f)
+				[
+					// import button
+					SNew(SButton)
+					.IsEnabled(this, &SSettingsSectionHeader::HandleImportButtonEnabled)
+					.OnClicked(this, &SSettingsSectionHeader::HandleImportButtonClicked)
+					.Text(LOCTEXT("ImportButtonText", "Import..."))
+					.ToolTipText(LOCTEXT("ImportButtonTooltip", "Import these settings from a file on your computer"))
+				]
+				+ SHorizontalBox::Slot()
+				.AutoWidth()
+				.Padding(8.0f, 0.0f, 0.0f, 0.0f)
+				[
+					// reset defaults button
+					SNew(SButton)
+					.Visibility(this, &SSettingsSectionHeader::HandleSetAsDefaultButtonVisibility)
+					.IsEnabled(this, &SSettingsSectionHeader::HandleResetToDefaultsButtonEnabled)
+					.OnClicked(this, &SSettingsSectionHeader::HandleResetDefaultsButtonClicked)
+					.Text(LOCTEXT("ResetDefaultsButtonText", "Reset to Defaults"))
+					.ToolTipText(LOCTEXT("ResetDefaultsButtonTooltip", "Reset the settings below to their default values"))
+				]
+			]
+		]
+		+ SVerticalBox::Slot()
+		.AutoHeight()
+		.Padding(0.0f, 8.0f, 0.0f, 0.0f)
+		[
+			FileWatcherWidget.ToSharedRef()
+		]
+	];
+
+}
+
+FText SSettingsSectionHeader::GetSettingsBoxTitleText() const
+{
+	if(SettingsSection.IsValid())
+	{
+		static const FText TitleFmt = FText::FromString(TEXT("{0} - {1}"));
+		return FText::Format(TitleFmt, SettingsSection->GetCategory().Pin()->GetDisplayName(), SettingsSection->GetDisplayName());
+	}
+
+	return FText::GetEmpty();
+}
+
+FText SSettingsSectionHeader::GetSettingsBoxDescriptionText() const
+{
+	if(SettingsSection.IsValid())
+	{
+		return SettingsSection->GetDescription();
+	}
+
+	return FText::GetEmpty();
+}
+
+EVisibility SSettingsSectionHeader::HandleButtonRowVisibility() const
+{
+	return DetailsView.Pin()->HasActiveSearch() ? EVisibility::Hidden : EVisibility::Visible;
+}
+
+FReply SSettingsSectionHeader::HandleExportButtonClicked()
+{
+	if(SettingsSection.IsValid())
+	{
+		if(LastExportDir.IsEmpty())
+		{
+			LastExportDir = FPaths::GetPath(GEditorPerProjectIni);
+		}
+
+		FString DefaultFileName = FString::Printf(TEXT("%s Backup %s.ini"), *SettingsSection->GetDisplayName().ToString(), *FDateTime::Now().ToString(TEXT("%Y-%m-%d %H%M%S")));
+		TArray<FString> OutFiles;
+
+		TSharedPtr<SWindow> ParentWindow = FSlateApplication::Get().FindWidgetWindow(AsShared());
+		void* ParentWindowHandle = (ParentWindow.IsValid() && ParentWindow->GetNativeWindow().IsValid()) ? ParentWindow->GetNativeWindow()->GetOSWindowHandle() : nullptr;
+
+		if(FDesktopPlatformModule::Get()->SaveFileDialog(ParentWindowHandle, LOCTEXT("ExportSettingsDialogTitle", "Export settings...").ToString(), LastExportDir, DefaultFileName, TEXT("Config files (*.ini)|*.ini"), EFileDialogFlags::None, OutFiles))
+		{
+			if(SettingsSection->Export(OutFiles[0]))
+			{
+				ShowNotification(LOCTEXT("ExportSettingsSuccess", "Export settings succeeded"), SNotificationItem::CS_Success);
+			}
+			else
+			{
+				ShowNotification(LOCTEXT("ExportSettingsFailure", "Export settings failed"), SNotificationItem::CS_Fail);
+			}
+		}
+	}
+
+	return FReply::Handled();
+}
+
+
+bool SSettingsSectionHeader::HandleExportButtonEnabled() const
+{
+	if(SettingsSection.IsValid())
+	{
+		return SettingsSection->CanExport();
+	}
+
+	return false;
+}
+
+
+FReply SSettingsSectionHeader::HandleImportButtonClicked()
+{
+	if(SettingsSection.IsValid())
+	{
+		TArray<FString> OutFiles;
+
+		TSharedPtr<SWindow> ParentWindow = FSlateApplication::Get().FindWidgetWindow(AsShared());
+		void* ParentWindowHandle = (ParentWindow.IsValid() && ParentWindow->GetNativeWindow().IsValid()) ? ParentWindow->GetNativeWindow()->GetOSWindowHandle() : nullptr;
+
+		if(FDesktopPlatformModule::Get()->OpenFileDialog(ParentWindowHandle, LOCTEXT("ImportSettingsDialogTitle", "Import settings...").ToString(), FPaths::GetPath(GEditorPerProjectIni), TEXT(""), TEXT("Config files (*.ini)|*.ini"), EFileDialogFlags::None, OutFiles))
+		{
+			if(SettingsSection->Import(OutFiles[0]) && SettingsSection->Save())
+			{
+				ShowNotification(LOCTEXT("ImportSettingsSuccess", "Import settings succeeded"), SNotificationItem::CS_Success);
+			}
+			else
+			{
+				ShowNotification(LOCTEXT("ImportSettingsFailure", "Import settings failed"), SNotificationItem::CS_Fail);
+			}
+		}
+	}
+
+	return FReply::Handled();
+}
+
+
+bool SSettingsSectionHeader::HandleImportButtonEnabled() const
+{
+	if(SettingsSection.IsValid())
+	{
+		return SettingsSection->CanEdit() && SettingsSection->CanImport() && !IsDefaultConfigCheckOutNeeded();
+	}
+
+	return false;
+}
+
+
+/**
+* Gets the absolute path to the Default.ini for the specified object.
+*
+* @return The path to the file.
+*/
+FString SSettingsSectionHeader::GetDefaultConfigFilePath() const
+{
+	FString RelativeConfigFilePath = SettingsObject->GetDefaultConfigFilename();
+	return FPaths::ConvertRelativePathToFull(RelativeConfigFilePath);
+}
+
+
+/**
+* Checks whether the default config file needs to be checked out for editing.
+*
+* @return true if the file needs to be checked out, false otherwise.
+*/
+bool SSettingsSectionHeader::IsDefaultConfigCheckOutNeeded() const
+{
+	if(SettingsObject.IsValid() && SettingsObject->GetClass()->HasAnyClassFlags(CLASS_Config | CLASS_DefaultConfig))
+	{
+		return !FileWatcherWidget->IsUnlocked();
+	}
+	else
+	{
+		return false;
+	}
+}
+
+FReply SSettingsSectionHeader::HandleResetDefaultsButtonClicked()
+{
+	if(SettingsSection.IsValid())
+	{
+		SettingsSection->ResetDefaults();
+	}
+
+	return FReply::Handled();
+}
+
+
+bool SSettingsSectionHeader::HandleResetToDefaultsButtonEnabled() const
+{
+	if(SettingsSection.IsValid())
+	{
+		return (SettingsSection->CanEdit() && SettingsSection->CanResetDefaults());
+	}
+
+	return false;
+}
+
+EVisibility SSettingsSectionHeader::HandleSetAsDefaultButtonVisibility() const
+{
+	return (SettingsSection.IsValid() && SettingsSection->HasDefaultSettingsObject()) ? EVisibility::Collapsed : EVisibility::Visible;
+}
+
+FReply SSettingsSectionHeader::HandleSetAsDefaultButtonClicked()
+{
+	if(SettingsSection.IsValid())
+	{
+		if(FMessageDialog::Open(EAppMsgType::YesNo, LOCTEXT("SaveAsDefaultUserConfirm", "Are you sure you want to update the default settings?")) != EAppReturnType::Yes)
+		{
+			return FReply::Handled();
+		}
+
+		if(IsDefaultConfigCheckOutNeeded())
+		{
+			if(ISourceControlModule::Get().IsEnabled())
+			{
+				if(FMessageDialog::Open(EAppMsgType::YesNo, LOCTEXT("SaveAsDefaultNeedsCheckoutMessage", "The default configuration file for these settings is currently not checked out. Would you like to check it out from source control? \n\nIf checked into source control this would affect other developers.")) != EAppReturnType::Yes)
+				{
+					return FReply::Handled();
+				}
+
+				CheckOutDefaultConfigFile();
+			}
+			else
+			{
+				if(FMessageDialog::Open(EAppMsgType::YesNo, LOCTEXT("SaveAsDefaultsIsReadOnlyMessage", "The default configuration file for these settings is currently not writeable. Would you like to make it writable?")) != EAppReturnType::Yes)
+				{
+					return FReply::Handled();
+				}
+
+				MakeDefaultConfigFileWritable();
+			}
+		}
+
+		SettingsSection->SaveDefaults();
+
+		FMessageDialog::Open(EAppMsgType::Ok, LOCTEXT("SaveAsDefaultsSucceededMessage", "The default configuration file for these settings was updated successfully. \n\nIf checked into source control this would affect other developers."));
+	}
+
+	return FReply::Handled();
+}
+
+
+bool SSettingsSectionHeader::HandleSetAsDefaultButtonEnabled() const
+{
+	if(SettingsSection.IsValid())
+	{
+		return SettingsSection->CanSaveDefaults();
+	}
+
+	return false;
+}
+
+/**
+* Checks out the default configuration file for the currently selected settings object.
+*
+* @return true if the check-out succeeded, false otherwise.
+*/
+bool SSettingsSectionHeader::CheckOutDefaultConfigFile()
+{
+	if(!SettingsObject.IsValid())
+	{
+		return false;
+	}
+
+	FText ErrorMessage;
+
+	// check out configuration file
+	FString AbsoluteConfigFilePath = GetDefaultConfigFilePath();
+	TArray<FString> FilesToBeCheckedOut;
+	FilesToBeCheckedOut.Add(AbsoluteConfigFilePath);
+
+	ISourceControlProvider& SourceControlProvider = ISourceControlModule::Get().GetProvider();
+	FSourceControlStatePtr SourceControlState = SourceControlProvider.GetState(AbsoluteConfigFilePath, EStateCacheUsage::ForceUpdate);
+
+	if(SourceControlState.IsValid())
+	{
+		if(SourceControlState->IsDeleted())
+		{
+			ErrorMessage = LOCTEXT("ConfigFileMarkedForDeleteError", "Error: The configuration file is marked for deletion.");
+		}
+		else if(!SourceControlState->IsCurrent())
+		{
+			if(false)
+			{
+				if(SourceControlProvider.Execute(ISourceControlOperation::Create<FSync>(), FilesToBeCheckedOut) == ECommandResult::Succeeded)
+				{
+					SettingsObject->ReloadConfig();
+				}
+				else
+				{
+					ErrorMessage = LOCTEXT("FailedToSyncConfigFileError", "Error: Failed to sync the configuration file to head revision.");
+				}
+			}
+		}
+		else if(SourceControlState->CanCheckout() || SourceControlState->IsCheckedOutOther())
+		{
+			if(SourceControlProvider.Execute(ISourceControlOperation::Create<FCheckOut>(), FilesToBeCheckedOut) == ECommandResult::Failed)
+			{
+				ErrorMessage = LOCTEXT("FailedToCheckOutConfigFileError", "Error: Failed to check out the configuration file.");
+			}
+		}
+	}
+
+	// show errors, if any
+	if(!ErrorMessage.IsEmpty())
+	{
+		FMessageDialog::Open(EAppMsgType::Ok, ErrorMessage);
+
+		return false;
+	}
+
+	return true;
+}
+
+/**
+* Makes the default configuration file for the currently selected settings object writable.
+*
+* @return true if it was made writable, false otherwise.
+*/
+bool SSettingsSectionHeader::MakeDefaultConfigFileWritable()
+{
+	if(!SettingsObject.IsValid())
+	{
+		return false;
+	}
+
+	FString AbsoluteConfigFilePath = GetDefaultConfigFilePath();
+
+	return FPlatformFileManager::Get().GetPlatformFile().SetReadOnly(*AbsoluteConfigFilePath, false);
+}
+
+void SSettingsSectionHeader::ShowNotification(const FText& Text, SNotificationItem::ECompletionState CompletionState) const
+{
+	FNotificationInfo Notification(Text);
+	Notification.ExpireDuration = 3.f;
+	Notification.bUseSuccessFailIcons = false;
+
+	FSlateNotificationManager::Get().AddNotification(Notification)->SetCompletionState(CompletionState);
+}
+
+/** Returns the config file name currently being edited. */
+FString SSettingsSectionHeader::HandleCheckoutNoticeConfigFilePath() const
+{
+	if(SettingsObject.IsValid() && SettingsObject->GetClass()->HasAnyClassFlags(CLASS_Config | CLASS_DefaultConfig))
+	{
+		return GetDefaultConfigFilePath();
+	}
+	else
+	{
+		return FString();
+	}
+}
+
+/** Reloads the configuration object. */
+void SSettingsSectionHeader::HandleCheckoutNoticeFileProbablyModifiedExternally()
+{
+	if(SettingsObject.IsValid() && SettingsObject->GetClass()->HasAnyClassFlags(CLASS_Config | CLASS_DefaultConfig))
+	{
+		SettingsObject->ReloadConfig();
+	}
+}
+
+/** Callback for determining the visibility of the 'Locked' notice. */
+EVisibility SSettingsSectionHeader::HandleCheckoutNoticeVisibility() const
+{
+	if(SettingsObject.IsValid() && SettingsObject->GetClass()->HasAnyClassFlags(CLASS_DefaultConfig) && !DetailsView.Pin()->HasActiveSearch())
+	{
+		return EVisibility::Visible;
+	}
+
+	return EVisibility::Collapsed;
+}
+
+void SSettingsSectionHeader::OnSettingsSelectionChanged()
+{
+	FileWatcherWidget->Invalidate();
+}
+
+
+FSettingsDetailRootObjectCustomization::FSettingsDetailRootObjectCustomization(ISettingsEditorModelPtr InModel, TSharedRef<IDetailsView> InDetailsView)
+	: Model(InModel)
+	, DetailsView(InDetailsView)
+{
+}
+
+void FSettingsDetailRootObjectCustomization::Initialize()
+{
+	Model->OnSelectionChanged().AddSP(this, &FSettingsDetailRootObjectCustomization::OnSelectedSectionChanged);
+	// Call once to ensure we have a valid section object
+	OnSelectedSectionChanged();
+}
+
+TSharedPtr<SWidget> FSettingsDetailRootObjectCustomization::CustomizeObjectHeader(const UObject* InRootObject)
+{
+	return SNew(SSettingsSectionHeader, InRootObject, Model, DetailsView.Pin());
+}
+
+bool FSettingsDetailRootObjectCustomization::IsObjectVisible(const UObject* InRootObject) const
+{
+	return SelectedSectionObject == nullptr || SelectedSectionObject == InRootObject || DetailsView.Pin()->HasActiveSearch();
+}
+
+bool FSettingsDetailRootObjectCustomization::ShouldDisplayHeader(const UObject* InRootObject) const
+{
+	return true;
+}
+
+void FSettingsDetailRootObjectCustomization::OnSelectedSectionChanged()
+{
+	ISettingsSectionPtr SelectedSection = Model->GetSelectedSection();
+	if(SelectedSection.IsValid())
+	{
+		SelectedSectionObject = SelectedSection->GetSettingsObject();
+	}
+	else
+	{
+		SelectedSectionObject = nullptr;
+	}
+
+	DetailsView.Pin()->RefreshRootObjectVisibility();
+}
+
+#undef LOCTEXT_NAMESPACE
\ No newline at end of file
diff --git a/Engine/Source/Developer/SettingsEditor/Public/ISettingsEditorModel.h b/Engine/Source/Developer/SettingsEditor/Public/ISettingsEditorModel.h
--- a/Engine/Source/Developer/SettingsEditor/Public/ISettingsEditorModel.h
+++ b/Engine/Source/Developer/SettingsEditor/Public/ISettingsEditorModel.h
@@ -42,6 +42,14 @@ public:
 	 */
 	virtual void SelectSection( const TSharedPtr<ISettingsSection>& Section ) = 0;
 
+	/**
+	 * Gets the section interface for a specified section object
+	 *
+	 * @param SectionObject	The settings section object to get the section interface for
+	 *
+	 * @return The section interface for the section object or nullptr if it does not exist
+	 */
+	virtual TSharedPtr<ISettingsSection> GetSectionFromSectionObject(const UObject* SectionObject) const = 0;
 public:
 
 	/**
diff --git a/Engine/Source/Developer/SharedSettingsWidgets/Private/SSettingsEditorCheckoutNotice.cpp b/Engine/Source/Developer/SharedSettingsWidgets/Private/SSettingsEditorCheckoutNotice.cpp
--- a/Engine/Source/Developer/SharedSettingsWidgets/Private/SSettingsEditorCheckoutNotice.cpp
+++ b/Engine/Source/Developer/SharedSettingsWidgets/Private/SSettingsEditorCheckoutNotice.cpp
@@ -6,13 +6,80 @@
 #include "ISourceControlModule.h"
 #include "SWidgetSwitcher.h"
 #include "SThrobber.h"
-#include "SNotificationList.h"
 #include "NotificationManager.h"
-
+#include "SNotificationList.h"
 
 #define LOCTEXT_NAMESPACE "SSettingsEditorCheckoutNotice"
 
+namespace SettingsHelpers
+{
+	bool CheckOutFile(const FString& InFileToCheckOut)
+	{
+		FText ErrorMessage;
+		bool bSuccessfullyCheckedOutFile = false;
+		if(ISourceControlModule::Get().IsEnabled())
+		{
+			ISourceControlProvider& SourceControlProvider = ISourceControlModule::Get().GetProvider();
+			FSourceControlStatePtr SourceControlState = SourceControlProvider.GetState(InFileToCheckOut, EStateCacheUsage::Use);
+
+			TArray<FString> FilesToBeCheckedOut;
+			FilesToBeCheckedOut.Add(InFileToCheckOut);
+
+			if(SourceControlState.IsValid())
+			{
+				if(SourceControlState->IsDeleted())
+				{
+					ErrorMessage = LOCTEXT("ConfigFileMarkedForDeleteError", "Error: The configuration file is marked for deletion.");
+				}
+				// Note: Here we attempt to check out files that are read only even if the internal state says they cannot be checked out.  This is to work around cases were the file is reverted or checked in and the internal state has not been updated yet
+				else if(SourceControlState->CanCheckout() || SourceControlState->IsCheckedOutOther() || FPlatformFileManager::Get().GetPlatformFile().IsReadOnly(*InFileToCheckOut))
+				{
+					ECommandResult::Type CommandResult = SourceControlProvider.Execute(ISourceControlOperation::Create<FCheckOut>(), FilesToBeCheckedOut);
+					if(CommandResult == ECommandResult::Failed)
+					{
+						ErrorMessage = LOCTEXT("FailedToCheckOutConfigFileError", "Error: Failed to check out the configuration file.");
+					}
+					else if(CommandResult == ECommandResult::Cancelled)
+					{
+						ErrorMessage = LOCTEXT("CancelledCheckOutConfigFile", "Checkout was cancelled.  File will be marked writable.");
+					}
+					else
+					{
+						bSuccessfullyCheckedOutFile = true;
+					}
+				}
+			}
+
+			if(!ErrorMessage.IsEmpty())
+			{
+				// Show a notification that the file could not be checked out
+				FNotificationInfo CheckOutError(ErrorMessage);
+				CheckOutError.ExpireDuration = 3.0f;
+				FSlateNotificationManager::Get().AddNotification(CheckOutError);
+			}
+
+		}
+
+
+		return bSuccessfullyCheckedOutFile;
+	}
+
+	bool MakeWritable(const FString& InFileToMakeWritable)
+	{
+		bool bSuccess = FPlatformFileManager::Get().GetPlatformFile().SetReadOnly(*InFileToMakeWritable, false);
+		if(!bSuccess)
+		{
+			FText NotificationErrorText = FText::Format(LOCTEXT("FailedToMakeWritable", "Could not make {0} writable."), FText::FromString(InFileToMakeWritable));
+
+			FNotificationInfo MakeWritiableNotification(NotificationErrorText);
+			MakeWritiableNotification.ExpireDuration = 3.0f;
+
+			FSlateNotificationManager::Get().AddNotification(MakeWritiableNotification);
+		}
+		return bSuccess;
 
+	}
+}
 /* SSettingsEditorCheckoutNotice interface
  *****************************************************************************/
 
@@ -126,30 +193,14 @@ FReply SSettingsEditorCheckoutNotice::HandleCheckOutButtonClicked()
 {
 	FString TargetFilePath = ConfigFilePath.Get();
 
-	if (ISourceControlModule::Get().IsEnabled())
-	{
-		FText ErrorMessage;
-
-		if (!SourceControlHelpers::CheckoutOrMarkForAdd(TargetFilePath, FText::FromString(TargetFilePath), NULL, ErrorMessage))
-		{
-			FNotificationInfo Info(ErrorMessage);
-			Info.ExpireDuration = 3.0f;
-			FSlateNotificationManager::Get().AddNotification(Info);
-		}
-	}
-	else
+	FText ErrorMessage;
+	if(!SettingsHelpers::CheckOutFile(TargetFilePath))
 	{
-		if (!FPlatformFileManager::Get().GetPlatformFile().SetReadOnly(*TargetFilePath, false))
-		{
-			FText NotificationErrorText = FText::Format(LOCTEXT("FailedToMakeWritable", "Could not make {0} writable."), FText::FromString(TargetFilePath));
-
-			FNotificationInfo Info(NotificationErrorText);
-			Info.ExpireDuration = 3.0f;
-
-			FSlateNotificationManager::Get().AddNotification(Info);
-		}
+		FNotificationInfo Info(ErrorMessage);
+		Info.ExpireDuration = 3.0f;
+		FSlateNotificationManager::Get().AddNotification(Info);
 	}
-
+		
 	return FReply::Handled();
 }
 
@@ -274,3 +325,5 @@ bool SSettingsEditorCheckoutNotice::IsUnlocked() const
 }
 
 #undef LOCTEXT_NAMESPACE
+
+
diff --git a/Engine/Source/Developer/SharedSettingsWidgets/Public/SSettingsEditorCheckoutNotice.h b/Engine/Source/Developer/SharedSettingsWidgets/Public/SSettingsEditorCheckoutNotice.h
--- a/Engine/Source/Developer/SharedSettingsWidgets/Public/SSettingsEditorCheckoutNotice.h
+++ b/Engine/Source/Developer/SharedSettingsWidgets/Public/SSettingsEditorCheckoutNotice.h
@@ -2,6 +2,11 @@
 
 #pragma once
 
+namespace SettingsHelpers
+{
+	SHAREDSETTINGSWIDGETS_API bool CheckOutFile(const FString& InFileToCheckOut);
+	SHAREDSETTINGSWIDGETS_API bool MakeWritable(const FString& InFileToMakeWritable);
+}
 
 /**
  * Implements a widget that provides a notice for files that need to be checked out.
diff --git a/Engine/Source/Developer/SourceControl/Private/SourceControlModule.cpp b/Engine/Source/Developer/SourceControl/Private/SourceControlModule.cpp
--- a/Engine/Source/Developer/SourceControl/Private/SourceControlModule.cpp
+++ b/Engine/Source/Developer/SourceControl/Private/SourceControlModule.cpp
@@ -210,6 +210,18 @@ void FSourceControlModule::InitializeSourceControlProviders()
 	CurrentSourceControlProvider->Init(false);	// Don't force a connection here, as its synchronous. Let the user establish a connection.
 }
 
+void FSourceControlModule::GetProviderNames(TArray<FName>& OutProviderNames)
+{
+	OutProviderNames.Reset();
+
+	int32 ProviderCount = GetNumSourceControlProviders();
+	for ( int32 ProviderIndex = 0; ProviderIndex < ProviderCount; ProviderIndex++ )
+	{
+		FName ProviderName = GetSourceControlProviderName(ProviderIndex);
+		OutProviderNames.Add(ProviderName);
+	}
+}
+
 void FSourceControlModule::Tick()
 {	
 	if( CurrentSourceControlProvider != nullptr )
diff --git a/Engine/Source/Developer/SourceControl/Private/SourceControlModule.h b/Engine/Source/Developer/SourceControl/Private/SourceControlModule.h
--- a/Engine/Source/Developer/SourceControl/Private/SourceControlModule.h
+++ b/Engine/Source/Developer/SourceControl/Private/SourceControlModule.h
@@ -22,6 +22,7 @@ public:
 	virtual void ShutdownModule() override;
 
 	/** ISourceControlModule implementation */
+	virtual void GetProviderNames(TArray<FName>& OutProviderNames) override;
 	virtual void Tick() override;
 	virtual void QueueStatusUpdate(const TArray<UPackage*>& InPackages) override;
 	virtual void QueueStatusUpdate(const TArray<FString>& InFilenames) override;
diff --git a/Engine/Source/Developer/SourceControl/Public/ISourceControlModule.h b/Engine/Source/Developer/SourceControl/Public/ISourceControlModule.h
--- a/Engine/Source/Developer/SourceControl/Public/ISourceControlModule.h
+++ b/Engine/Source/Developer/SourceControl/Public/ISourceControlModule.h
@@ -46,6 +46,11 @@ class ISourceControlModule : public IModuleInterface
 {
 public:
 
+	/**
+	 * Returns a list of the registered source control providers
+	 */
+	virtual void GetProviderNames(TArray<FName>& OutProviderNames) = 0;
+
 	/**
 	 * Tick the source control module.
 	 * This is responsible for dispatching batched/queued status requests & for calling ISourceControlProvider::Tick()
diff --git a/Engine/Source/Developer/SourceControl/Public/ISourceControlState.h b/Engine/Source/Developer/SourceControl/Public/ISourceControlState.h
--- a/Engine/Source/Developer/SourceControl/Public/ISourceControlState.h
+++ b/Engine/Source/Developer/SourceControl/Public/ISourceControlState.h
@@ -119,6 +119,9 @@ public:
 	/** Get whether source control allows this file to be edited */
 	virtual bool CanEdit() const = 0;
 
+	/** Get whether source control allows this file to be deleted. */
+	virtual bool CanDelete() const = 0;
+
 	/** Get whether we know anything about this files source control state */
 	virtual bool IsUnknown() const = 0;
 
diff --git a/Engine/Source/Editor/BlueprintGraph/Classes/EdGraphSchema_K2.h b/Engine/Source/Editor/BlueprintGraph/Classes/EdGraphSchema_K2.h
--- a/Engine/Source/Editor/BlueprintGraph/Classes/EdGraphSchema_K2.h
+++ b/Engine/Source/Editor/BlueprintGraph/Classes/EdGraphSchema_K2.h
@@ -406,6 +406,8 @@ public:
 	};
 
 public:
+	void SelectAllInputNodes(UEdGraph* Graph, UEdGraphPin* InGraphPin);
+
 	//~ Begin EdGraphSchema Interface
 	virtual void GetContextMenuActions(const UEdGraph* CurrentGraph, const UEdGraphNode* InGraphNode, const UEdGraphPin* InGraphPin, class FMenuBuilder* MenuBuilder, bool bIsDebugging) const override;
 	virtual const FPinConnectionResponse CanCreateConnection(const UEdGraphPin* A, const UEdGraphPin* B) const override;
diff --git a/Engine/Source/Editor/BlueprintGraph/Private/EdGraphSchema_K2.cpp b/Engine/Source/Editor/BlueprintGraph/Private/EdGraphSchema_K2.cpp
--- a/Engine/Source/Editor/BlueprintGraph/Private/EdGraphSchema_K2.cpp
+++ b/Engine/Source/Editor/BlueprintGraph/Private/EdGraphSchema_K2.cpp
@@ -1293,6 +1293,35 @@ bool UEdGraphSchema_K2::PinDefaultValueIsEditable(const UEdGraphPin& InGraphPin)
 	return true;
 }
 
+void UEdGraphSchema_K2::SelectAllInputNodes(UEdGraph* Graph, UEdGraphPin* InGraphPin)
+{
+	TArray<UEdGraphPin*> AllPins = InGraphPin->LinkedTo;
+
+	if (AllPins.Num() == 0)
+	{
+		return;
+	}
+
+	for (UEdGraphPin* Pin : AllPins)
+	{
+		UEdGraphNode* OwningNode = Pin->GetOwningNode();
+		FKismetEditorUtilities::AddToSelection(Graph, OwningNode);
+
+		TArray<UEdGraphPin*> LinkedPins = Pin->GetOwningNode()->GetAllPins();
+		for (UEdGraphPin* InputPin : LinkedPins)
+		{
+			if (InputPin->Direction == EEdGraphPinDirection::EGPD_Output)
+			{
+				continue;
+			}
+			else
+			{
+				SelectAllInputNodes(Graph, InputPin);
+			}
+		}
+	}
+}
+
 void UEdGraphSchema_K2::GetContextMenuActions(const UEdGraph* CurrentGraph, const UEdGraphNode* InGraphNode, const UEdGraphPin* InGraphPin, FMenuBuilder* MenuBuilder, bool bIsDebugging) const
 {
 	check(CurrentGraph);
@@ -1319,6 +1348,12 @@ void UEdGraphSchema_K2::GetContextMenuActions(const UEdGraph* CurrentGraph, cons
 				// add sub menu for break link to
 				if (InGraphPin->LinkedTo.Num() > 0)
 				{
+					MenuBuilder->AddMenuEntry(
+						LOCTEXT("SelectAllInputNodes", "Select All Input Nodes"),
+						LOCTEXT("SelectAllInputNodesTooltip", "Adds all input Nodes linked to this Pin to selection"),
+						FSlateIcon(),
+						FUIAction(FExecuteAction::CreateUObject((UEdGraphSchema_K2*const)this, &UEdGraphSchema_K2::SelectAllInputNodes, const_cast<UEdGraph*>(CurrentGraph), const_cast<UEdGraphPin*>(InGraphPin))));
+
 					if(InGraphPin->LinkedTo.Num() > 1)
 					{
 						MenuBuilder->AddSubMenu(
diff --git a/Engine/Source/Editor/ComponentVisualizers/Private/SplineComponentVisualizer.cpp b/Engine/Source/Editor/ComponentVisualizers/Private/SplineComponentVisualizer.cpp
--- a/Engine/Source/Editor/ComponentVisualizers/Private/SplineComponentVisualizer.cpp
+++ b/Engine/Source/Editor/ComponentVisualizers/Private/SplineComponentVisualizer.cpp
@@ -3,6 +3,7 @@
 #include "ComponentVisualizersPrivatePCH.h"
 #include "SplineComponentVisualizer.h"
 #include "ScopedTransaction.h"
+#include "ActorEditorUtils.h"
 
 IMPLEMENT_HIT_PROXY(HSplineVisProxy, HComponentVisProxy);
 IMPLEMENT_HIT_PROXY(HSplineKeyProxy, HSplineVisProxy);
@@ -37,6 +38,7 @@ public:
 		UI_COMMAND(SetKeyToLinear, "Linear", "Set spline point to Linear type", EUserInterfaceActionType::RadioButton, FInputChord());
 		UI_COMMAND(SetKeyToConstant, "Constant", "Set spline point to Constant type", EUserInterfaceActionType::RadioButton, FInputChord());
 		UI_COMMAND(VisualizeRollAndScale, "Visualize Roll and Scale", "Whether the visualization should show roll and scale on this spline.", EUserInterfaceActionType::ToggleButton, FInputChord());
+		UI_COMMAND(DiscontinuousSpline, "Allow Discontinuous Splines", "Whether the visualization allows Arrive and Leave tangents to be set separately.", EUserInterfaceActionType::ToggleButton, FInputChord());
 		UI_COMMAND(ResetToDefault, "Reset to Default", "Reset this spline to its archetype default.", EUserInterfaceActionType::Button, FInputChord());
 	}
 
@@ -68,6 +70,9 @@ public:
 	/** Whether the visualization should show roll and scale */
 	TSharedPtr<FUICommandInfo> VisualizeRollAndScale;
 
+	/** Whether we allow separate Arrive / Leave tangents, resulting in a discontinuous spline */
+	TSharedPtr<FUICommandInfo> DiscontinuousSpline;
+
 	/** Reset this spline to its default */
 	TSharedPtr<FUICommandInfo> ResetToDefault;
 };
@@ -85,6 +90,8 @@ FSplineComponentVisualizer::FSplineComponentVisualizer()
 	FSplineComponentVisualizerCommands::Register();
 
 	SplineComponentVisualizerActions = MakeShareable(new FUICommandList);
+
+	SplineCurvesProperty = FindField<UProperty>(USplineComponent::StaticClass(), GET_MEMBER_NAME_CHECKED(USplineComponent, SplineCurves));
 }
 
 void FSplineComponentVisualizer::OnRegister()
@@ -140,6 +147,12 @@ void FSplineComponentVisualizer::OnRegister()
 		FCanExecuteAction(),
 		FIsActionChecked::CreateSP(this, &FSplineComponentVisualizer::IsVisualizingRollAndScale));
 
+	SplineComponentVisualizerActions->MapAction(
+		Commands.DiscontinuousSpline,
+		FExecuteAction::CreateSP(this, &FSplineComponentVisualizer::OnSetDiscontinuousSpline),
+		FCanExecuteAction(),
+		FIsActionChecked::CreateSP(this, &FSplineComponentVisualizer::IsDiscontinuousSpline));
+
 	SplineComponentVisualizerActions->MapAction(
 		Commands.ResetToDefault,
 		FExecuteAction::CreateSP(this, &FSplineComponentVisualizer::OnResetToDefault),
@@ -169,10 +182,11 @@ void FSplineComponentVisualizer::DrawVisualization(const UActorComponent* Compon
 {
 	if (const USplineComponent* SplineComp = Cast<const USplineComponent>(Component))
 	{
+		const FInterpCurveVector& SplineInfo = SplineComp->GetSplinePointsPosition();
 		const USplineComponent* EditedSplineComp = GetEditedSplineComponent();
 
 		const USplineComponent* Archetype = CastChecked<USplineComponent>(SplineComp->GetArchetype());
-		const bool bIsSplineEditable = SplineComp->bSplineHasBeenEdited || SplineComp->SplineInfo.Points == Archetype->SplineInfo.Points || SplineComp->bInputSplinePointsToConstructionScript;
+		const bool bIsSplineEditable = !SplineComp->bModifiedByConstructionScript; // bSplineHasBeenEdited || SplineInfo == Archetype->SplineCurves.Position || SplineComp->bInputSplinePointsToConstructionScript;
 
 		const FColor ReadOnlyColor = FColor(255, 0, 255, 255);
 		const FColor NormalColor = bIsSplineEditable ? FColor(SplineComp->EditorUnselectedSplineSegmentColor.ToFColor(true)) : ReadOnlyColor;
@@ -180,8 +194,6 @@ void FSplineComponentVisualizer::DrawVisualization(const UActorComponent* Compon
 		const float GrabHandleSize = 12.0f;
 		const float TangentHandleSize = 10.0f;
 
-		const FInterpCurveVector& SplineInfo = SplineComp->SplineInfo;
-
 		// Draw the tangent handles before anything else so they will not overdraw the rest of the spline
 		if (SplineComp == EditedSplineComp)
 		{
@@ -190,33 +202,35 @@ void FSplineComponentVisualizer::DrawVisualization(const UActorComponent* Compon
 				if (SplineInfo.Points[SelectedKey].IsCurveKey())
 				{
 					const FVector Location = SplineComp->GetLocationAtSplinePoint(SelectedKey, ESplineCoordinateSpace::World);
-					const FVector Tangent = SplineComp->GetTangentAtSplinePoint(SelectedKey, ESplineCoordinateSpace::World);
+					const FVector LeaveTangent = SplineComp->GetLeaveTangentAtSplinePoint(SelectedKey, ESplineCoordinateSpace::World);
+					const FVector ArriveTangent = SplineComp->bAllowDiscontinuousSpline ?
+						SplineComp->GetArriveTangentAtSplinePoint(SelectedKey, ESplineCoordinateSpace::World) : LeaveTangent;
 
 					PDI->SetHitProxy(NULL);
 
-					const float DashSize1 = GetDashSize(View, Location, Location + Tangent, 0.01f);
+					const float DashSize1 = GetDashSize(View, Location, Location + LeaveTangent, 0.01f);
 					if (DashSize1 > 0.0f)
 					{
-						DrawDashedLine(PDI, Location, Location + Tangent, SelectedColor, DashSize1, SDPG_Foreground);
+						DrawDashedLine(PDI, Location, Location + LeaveTangent, SelectedColor, DashSize1, SDPG_Foreground);
 					}
 
-					const float DashSize2 = GetDashSize(View, Location, Location - Tangent, 0.01f);
+					const float DashSize2 = GetDashSize(View, Location, Location - ArriveTangent, 0.01f);
 					if (DashSize2 > 0.0f)
 					{
-						DrawDashedLine(PDI, Location, Location - Tangent, SelectedColor, DashSize2, SDPG_Foreground);
+						DrawDashedLine(PDI, Location, Location - ArriveTangent, SelectedColor, DashSize2, SDPG_Foreground);
 					}
 
 					if (bIsSplineEditable)
 					{
 						PDI->SetHitProxy(new HSplineTangentHandleProxy(Component, SelectedKey, false));
 					}
-					PDI->DrawPoint(Location + Tangent, SelectedColor, TangentHandleSize, SDPG_Foreground);
+					PDI->DrawPoint(Location + LeaveTangent, SelectedColor, TangentHandleSize, SDPG_Foreground);
 
 					if (bIsSplineEditable)
 					{
 						PDI->SetHitProxy(new HSplineTangentHandleProxy(Component, SelectedKey, true));
 					}
-					PDI->DrawPoint(Location - Tangent, SelectedColor, TangentHandleSize, SDPG_Foreground);
+					PDI->DrawPoint(Location - ArriveTangent, SelectedColor, TangentHandleSize, SDPG_Foreground);
 
 					PDI->SetHitProxy(NULL);
 				}
@@ -231,7 +245,7 @@ void FSplineComponentVisualizer::DrawVisualization(const UActorComponent* Compon
 		FVector OldKeyScale(0);
 
 		const int32 NumPoints = SplineInfo.Points.Num();
-		const int32 NumSegments = SplineComp->IsClosedLoop() ? NumPoints : NumPoints - 1;
+		const int32 NumSegments = SplineInfo.bIsLooped ? NumPoints : NumPoints - 1;
 		for (int32 KeyIdx = 0; KeyIdx < NumSegments + 1; KeyIdx++)
 		{
 			const FVector NewKeyPos = SplineComp->GetLocationAtSplinePoint(KeyIdx, ESplineCoordinateSpace::World);
@@ -381,7 +395,7 @@ void FSplineComponentVisualizer::ChangeSelectionState(int32 Index, bool bIsCtrlH
 	}
 }
 
-bool FSplineComponentVisualizer::VisProxyHandleClick(FLevelEditorViewportClient* InViewportClient, HComponentVisProxy* VisProxy, const FViewportClick& Click)
+bool FSplineComponentVisualizer::VisProxyHandleClick(FEditorViewportClient* InViewportClient, HComponentVisProxy* VisProxy, const FViewportClick& Click)
 {
 	if(VisProxy && VisProxy->Component.IsValid())
 	{
@@ -515,11 +529,13 @@ bool FSplineComponentVisualizer::GetWidgetLocation(const FEditorViewportClient*
 	USplineComponent* SplineComp = GetEditedSplineComponent();
 	if (SplineComp != nullptr)
 	{
+		const FInterpCurveVector& Position = SplineComp->GetSplinePointsPosition();
+
 		if (SelectedTangentHandle != INDEX_NONE)
 		{
 			// If tangent handle index is set, use that
-			check(SelectedTangentHandle < SplineComp->SplineInfo.Points.Num());
-			const auto& Point = SplineComp->SplineInfo.Points[SelectedTangentHandle];
+			check(SelectedTangentHandle < Position.Points.Num());
+			const auto& Point = Position.Points[SelectedTangentHandle];
 
 			check(SelectedTangentHandleType != ESelectedTangentHandle::None);
 			if (SelectedTangentHandleType == ESelectedTangentHandle::Leave)
@@ -536,9 +552,9 @@ bool FSplineComponentVisualizer::GetWidgetLocation(const FEditorViewportClient*
 		else if (LastKeyIndexSelected != INDEX_NONE)
 		{
 			// Otherwise use the last key index set
-			check(LastKeyIndexSelected < SplineComp->SplineInfo.Points.Num());
+			check(LastKeyIndexSelected < Position.Points.Num());
 			check(SelectedKeys.Contains(LastKeyIndexSelected));
-			const auto& Point = SplineComp->SplineInfo.Points[LastKeyIndexSelected];
+			const auto& Point = Position.Points[LastKeyIndexSelected];
 			OutLocation = SplineComp->ComponentToWorld.TransformPosition(Point.OutVal);
 			return true;
 		}
@@ -564,16 +580,23 @@ bool FSplineComponentVisualizer::GetCustomInputCoordinateSystem(const FEditorVie
 }
 
 
+bool FSplineComponentVisualizer::IsVisualizingArchetype() const
+{
+	USplineComponent* SplineComp = GetEditedSplineComponent();
+	return (SplineComp && SplineComp->GetOwner() && FActorEditorUtils::IsAPreviewOrInactiveActor(SplineComp->GetOwner()));
+}
+
+
 bool FSplineComponentVisualizer::HandleInputDelta(FEditorViewportClient* ViewportClient, FViewport* Viewport, FVector& DeltaTranslate, FRotator& DeltaRotate, FVector& DeltaScale)
 {
 	USplineComponent* SplineComp = GetEditedSplineComponent();
 	if (SplineComp != nullptr)
 	{
-		FInterpCurveVector& SplineInfo = SplineComp->SplineInfo;
-		FInterpCurveQuat& SplineRotInfo = SplineComp->SplineRotInfo;
-		FInterpCurveVector& SplineScaleInfo = SplineComp->SplineScaleInfo;
+		FInterpCurveVector& SplinePosition = SplineComp->GetSplinePointsPosition();
+		FInterpCurveQuat& SplineRotation = SplineComp->GetSplinePointsRotation();
+		FInterpCurveVector& SplineScale = SplineComp->GetSplinePointsScale();
 
-		const int32 NumPoints = SplineInfo.Points.Num();
+		const int32 NumPoints = SplinePosition.Points.Num();
 
 		if (SelectedTangentHandle != INDEX_NONE)
 		{
@@ -587,12 +610,27 @@ bool FSplineComponentVisualizer::HandleInputDelta(FEditorViewportClient* Viewpor
 
 				SplineComp->Modify();
 
-				FInterpCurvePoint<FVector>& EditedPoint = SplineInfo.Points[SelectedTangentHandle];
-				const FVector Delta = (SelectedTangentHandleType == ESelectedTangentHandle::Leave) ? DeltaTranslate : -DeltaTranslate;
-				const FVector Tangent = EditedPoint.LeaveTangent + SplineComp->ComponentToWorld.InverseTransformVector(Delta);
+				FInterpCurvePoint<FVector>& EditedPoint = SplinePosition.Points[SelectedTangentHandle];
+				if (SplineComp->bAllowDiscontinuousSpline)
+				{
+					if (SelectedTangentHandleType == ESelectedTangentHandle::Leave)
+					{
+						EditedPoint.LeaveTangent += SplineComp->ComponentToWorld.InverseTransformVector(DeltaTranslate);
+					}
+					else
+					{
+						EditedPoint.ArriveTangent += SplineComp->ComponentToWorld.InverseTransformVector(-DeltaTranslate);
+					}
+				}
+				else
+				{
+					const FVector Delta = (SelectedTangentHandleType == ESelectedTangentHandle::Leave) ? DeltaTranslate : -DeltaTranslate;
+					const FVector Tangent = EditedPoint.LeaveTangent + SplineComp->ComponentToWorld.InverseTransformVector(Delta);
+
+					EditedPoint.LeaveTangent = Tangent;
+					EditedPoint.ArriveTangent = Tangent;
+				}
 
-				EditedPoint.LeaveTangent = Tangent;
-				EditedPoint.ArriveTangent = Tangent;
 				EditedPoint.InterpMode = CIM_CurveUser;
 			}
 		}
@@ -608,16 +646,17 @@ bool FSplineComponentVisualizer::HandleInputDelta(FEditorViewportClient* Viewpor
 
 			if (ViewportClient->IsAltPressed() && bAllowDuplication)
 			{
-				OnDuplicateKey();
+				DuplicateKey();
+
 				// Don't duplicate again until we release LMB
 				bAllowDuplication = false;
 			}
 
 			for (int32 SelectedKeyIndex : SelectedKeys)
 			{
-				FInterpCurvePoint<FVector>& EditedPoint = SplineInfo.Points[SelectedKeyIndex];
-				FInterpCurvePoint<FQuat>& EditedRotPoint = SplineRotInfo.Points[SelectedKeyIndex];
-				FInterpCurvePoint<FVector>& EditedScalePoint = SplineScaleInfo.Points[SelectedKeyIndex];
+				FInterpCurvePoint<FVector>& EditedPoint = SplinePosition.Points[SelectedKeyIndex];
+				FInterpCurvePoint<FQuat>& EditedRotPoint = SplineRotation.Points[SelectedKeyIndex];
+				FInterpCurvePoint<FVector>& EditedScalePoint = SplineScale.Points[SelectedKeyIndex];
 
 				if (!DeltaTranslate.IsZero())
 				{
@@ -672,7 +711,10 @@ bool FSplineComponentVisualizer::HandleInputDelta(FEditorViewportClient* Viewpor
 			}
 		}
 
-		NotifyComponentModified();
+		SplineComp->UpdateSpline();
+		SplineComp->bSplineHasBeenEdited = true;
+
+		NotifyPropertyModified(SplineComp, SplineCurvesProperty);
 		return true;
 	}
 
@@ -726,6 +768,19 @@ void FSplineComponentVisualizer::EndEditing()
 void FSplineComponentVisualizer::OnDuplicateKey()
 {
 	const FScopedTransaction Transaction(LOCTEXT("DuplicateSplinePoint", "Duplicate Spline Point"));
+	
+	USplineComponent* SplineComp = GetEditedSplineComponent();
+	DuplicateKey();
+
+	SplineComp->UpdateSpline();
+	SplineComp->bSplineHasBeenEdited = true;
+
+	NotifyPropertyModified(SplineComp, SplineCurvesProperty);
+}
+
+
+void FSplineComponentVisualizer::DuplicateKey()
+{
 	USplineComponent* SplineComp = GetEditedSplineComponent();
 	check(SplineComp != nullptr);
 	check(LastKeyIndexSelected != INDEX_NONE);
@@ -747,17 +802,25 @@ void FSplineComponentVisualizer::OnDuplicateKey()
 	SelectedKeysSorted.Sort([](int32 A, int32 B) { return A > B; });
 
 	// Insert duplicates into the list, highest index first, so that the lower indices remain the same
-	FInterpCurveVector& SplineInfo = SplineComp->SplineInfo;
-	FInterpCurveQuat& SplineRotInfo = SplineComp->SplineRotInfo;
-	FInterpCurveVector& SplineScaleInfo = SplineComp->SplineScaleInfo;
+	FInterpCurveVector& SplinePosition = SplineComp->SplineCurves.Position;
+	FInterpCurveQuat& SplineRotation = SplineComp->SplineCurves.Rotation;
+	FInterpCurveVector& SplineScale = SplineComp->SplineCurves.Scale;
 
 	for (int32 SelectedKeyIndex : SelectedKeysSorted)
 	{
 		// Insert duplicates into arrays.
 		// It's necessary to take a copy because copying existing array items by reference isn't allowed (the array may reallocate)
-		SplineInfo.Points.Insert(FInterpCurvePoint<FVector>(SplineInfo.Points[SelectedKeyIndex]), SelectedKeyIndex);
-		SplineRotInfo.Points.Insert(FInterpCurvePoint<FQuat>(SplineRotInfo.Points[SelectedKeyIndex]), SelectedKeyIndex);
-		SplineScaleInfo.Points.Insert(FInterpCurvePoint<FVector>(SplineScaleInfo.Points[SelectedKeyIndex]), SelectedKeyIndex);
+		SplinePosition.Points.Insert(FInterpCurvePoint<FVector>(SplinePosition.Points[SelectedKeyIndex]), SelectedKeyIndex);
+		SplineRotation.Points.Insert(FInterpCurvePoint<FQuat>(SplineRotation.Points[SelectedKeyIndex]), SelectedKeyIndex);
+		SplineScale.Points.Insert(FInterpCurvePoint<FVector>(SplineScale.Points[SelectedKeyIndex]), SelectedKeyIndex);
+
+		// Adjust input keys of subsequent points
+		for (int Index = SelectedKeyIndex + 1; Index < SplinePosition.Points.Num(); Index++)
+		{
+			SplinePosition.Points[Index].InVal += 1.0f;
+			SplineRotation.Points[Index].InVal += 1.0f;
+			SplineScale.Points[Index].InVal += 1.0f;
+		}
 	}
 
 	// Repopulate the selected keys
@@ -778,10 +841,6 @@ void FSplineComponentVisualizer::OnDuplicateKey()
 	// Unset tangent handle selection
 	SelectedTangentHandle = INDEX_NONE;
 	SelectedTangentHandleType = ESelectedTangentHandle::None;
-
-	NotifyComponentModified();
-
-//	CachedRotation = SplineComp->GetQuaternionAtSplinePoint(LastKeyIndexSelected, ESplineCoordinateSpace::World);
 }
 
 
@@ -793,7 +852,7 @@ bool FSplineComponentVisualizer::CanAddKey() const
 		return false;
 	}
 
-	const int32 NumPoints = SplineComp->SplineInfo.Points.Num();
+	const int32 NumPoints = SplineComp->SplineCurves.Position.Points.Num();
 	const int32 NumSegments = SplineComp->IsClosedLoop() ? NumPoints : NumPoints - 1;
 
 	return (SelectedSegmentIndex != INDEX_NONE && SelectedSegmentIndex < NumSegments);
@@ -817,23 +876,51 @@ void FSplineComponentVisualizer::OnAddKey()
 		Owner->Modify();
 	}
 
-	FInterpCurveVector& SplineInfo = SplineComp->SplineInfo;
-	FInterpCurveQuat& SplineRotInfo = SplineComp->SplineRotInfo;
-	FInterpCurveVector& SplineScaleInfo = SplineComp->SplineScaleInfo;
-
-	FInterpCurvePoint<FVector> NewPoint(0.0f, SplineComp->ComponentToWorld.InverseTransformPosition(SelectedSplinePosition), FVector::ZeroVector, FVector::ZeroVector, CIM_CurveAuto);
-	FInterpCurvePoint<FQuat> NewRotPoint(0.0f, FQuat::Identity, FQuat::Identity, FQuat::Identity, CIM_CurveAuto);
-	FInterpCurvePoint<FVector> NewScalePoint(0.0f, FVector(1.0f), FVector::ZeroVector, FVector::ZeroVector, CIM_CurveAuto);
-
-	SplineInfo.Points.Insert(NewPoint, SelectedSegmentIndex + 1);
-	SplineRotInfo.Points.Insert(NewRotPoint, SelectedSegmentIndex + 1);
-	SplineScaleInfo.Points.Insert(NewScalePoint, SelectedSegmentIndex + 1);
+	FInterpCurveVector& SplinePosition = SplineComp->GetSplinePointsPosition();
+	FInterpCurveQuat& SplineRotation = SplineComp->GetSplinePointsRotation();
+	FInterpCurveVector& SplineScale = SplineComp->GetSplinePointsScale();
+
+	FInterpCurvePoint<FVector> NewPoint(
+		SelectedSegmentIndex,
+		SplineComp->ComponentToWorld.InverseTransformPosition(SelectedSplinePosition),
+		FVector::ZeroVector,
+		FVector::ZeroVector,
+		CIM_CurveAuto);
+
+	FInterpCurvePoint<FQuat> NewRotPoint(
+		SelectedSegmentIndex,
+		FQuat::Identity,
+		FQuat::Identity,
+		FQuat::Identity,
+		CIM_CurveAuto);
+
+	FInterpCurvePoint<FVector> NewScalePoint(
+		SelectedSegmentIndex,
+		FVector(1.0f),
+		FVector::ZeroVector,
+		FVector::ZeroVector,
+		CIM_CurveAuto);
+
+	SplinePosition.Points.Insert(NewPoint, SelectedSegmentIndex + 1);
+	SplineRotation.Points.Insert(NewRotPoint, SelectedSegmentIndex + 1);
+	SplineScale.Points.Insert(NewScalePoint, SelectedSegmentIndex + 1);
+
+	// Adjust input keys of subsequent points
+	for (int Index = SelectedSegmentIndex + 1; Index < SplinePosition.Points.Num(); Index++)
+	{
+		SplinePosition.Points[Index].InVal += 1.0f;
+		SplineRotation.Points[Index].InVal += 1.0f;
+		SplineScale.Points[Index].InVal += 1.0f;
+	}
 
 	// Set selection to 'next' key
 	ChangeSelectionState(SelectedSegmentIndex + 1, false);
 	SelectedSegmentIndex = INDEX_NONE;
 
-	NotifyComponentModified();
+	SplineComp->UpdateSpline();
+	SplineComp->bSplineHasBeenEdited = true;
+
+	NotifyPropertyModified(SplineComp, SplineCurvesProperty);
 
 	CachedRotation = SplineComp->GetQuaternionAtSplinePoint(LastKeyIndexSelected, ESplineCoordinateSpace::World);
 }
@@ -863,15 +950,22 @@ void FSplineComponentVisualizer::OnDeleteKey()
 	SelectedKeysSorted.Sort([](int32 A, int32 B) { return A > B; });
 
 	// Delete selected keys from list, highest index first
-	FInterpCurveVector& SplineInfo = SplineComp->SplineInfo;
-	FInterpCurveQuat& SplineRotInfo = SplineComp->SplineRotInfo;
-	FInterpCurveVector& SplineScaleInfo = SplineComp->SplineScaleInfo;
+	FInterpCurveVector& SplinePosition = SplineComp->SplineCurves.Position;
+	FInterpCurveQuat& SplineRotation = SplineComp->SplineCurves.Rotation;
+	FInterpCurveVector& SplineScale = SplineComp->SplineCurves.Scale;
 
 	for (int32 SelectedKeyIndex : SelectedKeysSorted)
 	{
-		SplineInfo.Points.RemoveAt(SelectedKeyIndex);
-		SplineRotInfo.Points.RemoveAt(SelectedKeyIndex);
-		SplineScaleInfo.Points.RemoveAt(SelectedKeyIndex);
+		SplinePosition.Points.RemoveAt(SelectedKeyIndex);
+		SplineRotation.Points.RemoveAt(SelectedKeyIndex);
+		SplineScale.Points.RemoveAt(SelectedKeyIndex);
+
+		for (int Index = SelectedKeyIndex; Index < SplinePosition.Points.Num(); Index++)
+		{
+			SplinePosition.Points[Index].InVal -= 1.0f;
+			SplineRotation.Points[Index].InVal -= 1.0f;
+			SplineScale.Points[Index].InVal -= 1.0f;
+		}
 	}
 
 	// Select first key
@@ -880,7 +974,10 @@ void FSplineComponentVisualizer::OnDeleteKey()
 	SelectedTangentHandle = INDEX_NONE;
 	SelectedTangentHandleType = ESelectedTangentHandle::None;
 
-	NotifyComponentModified();
+	SplineComp->UpdateSpline();
+	SplineComp->bSplineHasBeenEdited = true;
+
+	NotifyPropertyModified(SplineComp, SplineCurvesProperty);
 
 	CachedRotation = SplineComp->GetQuaternionAtSplinePoint(LastKeyIndexSelected, ESplineCoordinateSpace::World);
 }
@@ -891,7 +988,7 @@ bool FSplineComponentVisualizer::CanDeleteKey() const
 	USplineComponent* SplineComp = GetEditedSplineComponent();
 	return (SplineComp != nullptr &&
 			SelectedKeys.Num() > 0 &&
-			SelectedKeys.Num() != SplineComp->SplineInfo.Points.Num() &&
+			SelectedKeys.Num() != SplineComp->SplineCurves.Position.Points.Num() &&
 			LastKeyIndexSelected != INDEX_NONE);
 }
 
@@ -920,14 +1017,17 @@ void FSplineComponentVisualizer::OnResetToAutomaticTangent(EInterpCurveMode Mode
 
 		for (int32 SelectedKeyIndex : SelectedKeys)
 		{
-			auto& Point = SplineComp->SplineInfo.Points[SelectedKeyIndex];
+			auto& Point = SplineComp->SplineCurves.Position.Points[SelectedKeyIndex];
 			if (Point.IsCurveKey())
 			{
 				Point.InterpMode = Mode;
 			}
 		}
 
-		NotifyComponentModified();
+		SplineComp->UpdateSpline();
+		SplineComp->bSplineHasBeenEdited = true;
+
+		NotifyPropertyModified(SplineComp, SplineCurvesProperty);
 
 		CachedRotation = SplineComp->GetQuaternionAtSplinePoint(LastKeyIndexSelected, ESplineCoordinateSpace::World);
 	}
@@ -941,7 +1041,7 @@ bool FSplineComponentVisualizer::CanResetToAutomaticTangent(EInterpCurveMode Mod
 	{
 		for (int32 SelectedKeyIndex : SelectedKeys)
 		{
-			const auto& Point = SplineComp->SplineInfo.Points[SelectedKeyIndex];
+			const auto& Point = SplineComp->SplineCurves.Position.Points[SelectedKeyIndex];
 			if (Point.IsCurveKey() && Point.InterpMode != Mode)
 			{
 				return true;
@@ -968,10 +1068,13 @@ void FSplineComponentVisualizer::OnSetKeyType(EInterpCurveMode Mode)
 
 		for (int32 SelectedKeyIndex : SelectedKeys)
 		{
-			SplineComp->SplineInfo.Points[SelectedKeyIndex].InterpMode = Mode;
+			SplineComp->SplineCurves.Position.Points[SelectedKeyIndex].InterpMode = Mode;
 		}
 
-		NotifyComponentModified();
+		SplineComp->UpdateSpline();
+		SplineComp->bSplineHasBeenEdited = true;
+
+		NotifyPropertyModified(SplineComp, SplineCurvesProperty);
 
 		CachedRotation = SplineComp->GetQuaternionAtSplinePoint(LastKeyIndexSelected, ESplineCoordinateSpace::World);
 	}
@@ -987,7 +1090,7 @@ bool FSplineComponentVisualizer::IsKeyTypeSet(EInterpCurveMode Mode) const
 
 		for (int32 SelectedKeyIndex : SelectedKeys)
 		{
-			const auto& SelectedPoint = SplineComp->SplineInfo.Points[SelectedKeyIndex];
+			const auto& SelectedPoint = SplineComp->SplineCurves.Position.Points[SelectedKeyIndex];
 			if ((Mode == CIM_CurveAuto && SelectedPoint.IsCurveKey()) || SelectedPoint.InterpMode == Mode)
 			{
 				return true;
@@ -1003,7 +1106,17 @@ void FSplineComponentVisualizer::OnSetVisualizeRollAndScale()
 {
 	USplineComponent* SplineComp = GetEditedSplineComponent();
 	check(SplineComp != nullptr);
+
+	SplineComp->Modify();
+	if (AActor* Owner = SplineComp->GetOwner())
+	{
+		Owner->Modify();
+	}
+
 	SplineComp->bShouldVisualizeScale = !SplineComp->bShouldVisualizeScale;
+
+	NotifyPropertyModified(SplineComp, FindField<UProperty>(USplineComponent::StaticClass(), GET_MEMBER_NAME_CHECKED(USplineComponent, bShouldVisualizeScale)));
+
 	GEditor->RedrawLevelEditingViewports(true);
 }
 
@@ -1016,6 +1129,45 @@ bool FSplineComponentVisualizer::IsVisualizingRollAndScale() const
 }
 
 
+void FSplineComponentVisualizer::OnSetDiscontinuousSpline()
+{
+	USplineComponent* SplineComp = GetEditedSplineComponent();
+	check(SplineComp != nullptr);
+
+	SplineComp->Modify();
+	if (AActor* Owner = SplineComp->GetOwner())
+	{
+		Owner->Modify();
+	}
+
+	SplineComp->bAllowDiscontinuousSpline = !SplineComp->bAllowDiscontinuousSpline;
+
+	// If not allowed discontinuous splines, set all ArriveTangents to match LeaveTangents
+	if (!SplineComp->bAllowDiscontinuousSpline)
+	{
+		for (int Index = 0; Index < SplineComp->SplineCurves.Position.Points.Num(); Index++)
+		{
+			SplineComp->SplineCurves.Position.Points[Index].ArriveTangent = SplineComp->SplineCurves.Position.Points[Index].LeaveTangent;
+		}
+	}
+
+	TArray<UProperty*> Properties;
+	Properties.Add(SplineCurvesProperty);
+	Properties.Add(FindField<UProperty>(USplineComponent::StaticClass(), GET_MEMBER_NAME_CHECKED(USplineComponent, bAllowDiscontinuousSpline)));
+	NotifyPropertiesModified(SplineComp, Properties);
+
+	GEditor->RedrawLevelEditingViewports(true);
+}
+
+
+bool FSplineComponentVisualizer::IsDiscontinuousSpline() const
+{
+	USplineComponent* SplineComp = GetEditedSplineComponent();
+
+	return SplineComp ? SplineComp->bAllowDiscontinuousSpline : false;
+}
+
+
 void FSplineComponentVisualizer::OnResetToDefault()
 {
 	USplineComponent* SplineComp = GetEditedSplineComponent();
@@ -1050,7 +1202,8 @@ bool FSplineComponentVisualizer::CanResetToDefault() const
 {
 	USplineComponent* SplineComp = GetEditedSplineComponent();
 	check(SplineComp != nullptr);
-	return SplineComp->bSplineHasBeenEdited;
+	return SplineComp->SplineCurves != CastChecked<USplineComponent>(SplineComp->GetArchetype())->SplineCurves;
+//	return SplineComp->bSplineHasBeenEdited;
 }
 
 
@@ -1080,7 +1233,7 @@ TSharedPtr<SWidget> FSplineComponentVisualizer::GenerateContextMenu() const
 				{
 					for (int32 SelectedKeyIndex : SelectedKeys)
 					{
-						const auto& Point = SplineComp->SplineInfo.Points[SelectedKeyIndex];
+						const auto& Point = SplineComp->SplineCurves.Position.Points[SelectedKeyIndex];
 						if (Point.IsCurveKey())
 						{
 							MenuBuilder.AddSubMenu(
@@ -1104,6 +1257,7 @@ TSharedPtr<SWidget> FSplineComponentVisualizer::GenerateContextMenu() const
 		MenuBuilder.BeginSection("Visualization", LOCTEXT("Visualization", "Visualization"));
 		{
 			MenuBuilder.AddMenuEntry(FSplineComponentVisualizerCommands::Get().VisualizeRollAndScale);
+			MenuBuilder.AddMenuEntry(FSplineComponentVisualizerCommands::Get().DiscontinuousSpline);
 		}
 		MenuBuilder.EndSection();
 	}
@@ -1127,26 +1281,4 @@ void FSplineComponentVisualizer::GenerateTangentTypeSubMenu(FMenuBuilder& MenuBu
 	MenuBuilder.AddMenuEntry(FSplineComponentVisualizerCommands::Get().ResetToClampedTangent);
 }
 
-
-void FSplineComponentVisualizer::NotifyComponentModified()
-{
-	USplineComponent* SplineComp = GetEditedSplineComponent();
-
-	SplineComp->bSplineHasBeenEdited = true;
-	SplineComp->UpdateSpline();
-
-	// Notify that the spline info has been modified
-	UProperty* SplineInfoProperty = FindField<UProperty>(USplineComponent::StaticClass(), GET_MEMBER_NAME_CHECKED(USplineComponent, SplineInfo));
-	FPropertyChangedEvent PropertyChangedEvent(SplineInfoProperty);
-	SplineComp->PostEditChangeProperty(PropertyChangedEvent);
-
-	// Notify of change so any CS is re-run
-	if (SplineOwningActor.IsValid())
-	{
-		SplineOwningActor.Get()->PostEditMove(false);
-	}
-
-	GEditor->RedrawLevelEditingViewports(true);
-}
-
 #undef LOCTEXT_NAMESPACE
diff --git a/Engine/Source/Editor/ComponentVisualizers/Private/SplineMeshComponentVisualizer.cpp b/Engine/Source/Editor/ComponentVisualizers/Private/SplineMeshComponentVisualizer.cpp
--- a/Engine/Source/Editor/ComponentVisualizers/Private/SplineMeshComponentVisualizer.cpp
+++ b/Engine/Source/Editor/ComponentVisualizers/Private/SplineMeshComponentVisualizer.cpp
@@ -99,7 +99,7 @@ FInterpCurveVector FSplineMeshComponentVisualizer::GetSpline(const USplineMeshCo
 }
 
 
-bool FSplineMeshComponentVisualizer::VisProxyHandleClick(FLevelEditorViewportClient* InViewportClient, HComponentVisProxy* VisProxy, const FViewportClick& Click)
+bool FSplineMeshComponentVisualizer::VisProxyHandleClick(FEditorViewportClient* InViewportClient, HComponentVisProxy* VisProxy, const FViewportClick& Click)
 {
 	bool bEditing = false;
 
diff --git a/Engine/Source/Editor/ComponentVisualizers/Public/SplineComponentVisualizer.h b/Engine/Source/Editor/ComponentVisualizers/Public/SplineComponentVisualizer.h
--- a/Engine/Source/Editor/ComponentVisualizers/Public/SplineComponentVisualizer.h
+++ b/Engine/Source/Editor/ComponentVisualizers/Public/SplineComponentVisualizer.h
@@ -66,24 +66,28 @@ public:
 	//~ Begin FComponentVisualizer Interface
 	virtual void OnRegister() override;
 	virtual void DrawVisualization(const UActorComponent* Component, const FSceneView* View, FPrimitiveDrawInterface* PDI) override;
-	virtual bool VisProxyHandleClick(FLevelEditorViewportClient* InViewportClient, HComponentVisProxy* VisProxy, const FViewportClick& Click) override;
+	virtual bool VisProxyHandleClick(FEditorViewportClient* InViewportClient, HComponentVisProxy* VisProxy, const FViewportClick& Click) override;
 	virtual void EndEditing() override;
 	virtual bool GetWidgetLocation(const FEditorViewportClient* ViewportClient, FVector& OutLocation) const override;
 	virtual bool GetCustomInputCoordinateSystem(const FEditorViewportClient* ViewportClient, FMatrix& OutMatrix) const override;
 	virtual bool HandleInputDelta(FEditorViewportClient* ViewportClient, FViewport* Viewport, FVector& DeltaTranslate, FRotator& DeltaRotate, FVector& DeltaScale) override;
 	virtual bool HandleInputKey(FEditorViewportClient* ViewportClient, FViewport* Viewport, FKey Key, EInputEvent Event) override;
 	virtual TSharedPtr<SWidget> GenerateContextMenu() const override;
+	virtual bool IsVisualizingArchetype() const override;
 	//~ End FComponentVisualizer Interface
 
 	/** Get the spline component we are currently editing */
 	USplineComponent* GetEditedSplineComponent() const;
 
+	TSet<int32> GetSelectedKeys() const { return SelectedKeys; }
+
 protected:
 
 	/** Update the key selection state of the visualizer */
 	void ChangeSelectionState(int32 Index, bool bIsCtrlHeld);
 
-	void NotifyComponentModified();
+	/** Duplicates the selected spline key(s) */
+	void DuplicateKey();
 
 	void OnDeleteKey();
 	bool CanDeleteKey() const;
@@ -103,6 +107,9 @@ protected:
 	void OnSetVisualizeRollAndScale();
 	bool IsVisualizingRollAndScale() const;
 
+	void OnSetDiscontinuousSpline();
+	bool IsDiscontinuousSpline() const;
+
 	void OnResetToDefault();
 	bool CanResetToDefault() const;
 
@@ -154,4 +161,7 @@ protected:
 
 	/** Whether we currently allow duplication when dragging */
 	bool bAllowDuplication;
+
+private:
+	UProperty* SplineCurvesProperty;
 };
diff --git a/Engine/Source/Editor/ComponentVisualizers/Public/SplineMeshComponentVisualizer.h b/Engine/Source/Editor/ComponentVisualizers/Public/SplineMeshComponentVisualizer.h
--- a/Engine/Source/Editor/ComponentVisualizers/Public/SplineMeshComponentVisualizer.h
+++ b/Engine/Source/Editor/ComponentVisualizers/Public/SplineMeshComponentVisualizer.h
@@ -53,7 +53,7 @@ public:
 	//~ Begin FComponentVisualizer Interface
 	virtual void OnRegister() override;
 	virtual void DrawVisualization(const UActorComponent* Component, const FSceneView* View, FPrimitiveDrawInterface* PDI) override;
-	virtual bool VisProxyHandleClick(FLevelEditorViewportClient* InViewportClient, HComponentVisProxy* VisProxy, const FViewportClick& Click) override;
+	virtual bool VisProxyHandleClick(FEditorViewportClient* InViewportClient, HComponentVisProxy* VisProxy, const FViewportClick& Click) override;
 	virtual void EndEditing() override;
 	virtual bool GetWidgetLocation(const FEditorViewportClient* ViewportClient, FVector& OutLocation) const override;
 	virtual bool GetCustomInputCoordinateSystem(const FEditorViewportClient* ViewportClient, FMatrix& OutMatrix) const override;
diff --git a/Engine/Source/Editor/ContentBrowser/Private/ContentBrowserUtils.cpp b/Engine/Source/Editor/ContentBrowser/Private/ContentBrowserUtils.cpp
--- a/Engine/Source/Editor/ContentBrowser/Private/ContentBrowserUtils.cpp
+++ b/Engine/Source/Editor/ContentBrowser/Private/ContentBrowserUtils.cpp
@@ -566,9 +566,13 @@ bool ContentBrowserUtils::DeleteFolders(const TArray<FString>& PathsToDelete)
 	{
 		FAssetRegistryModule& AssetRegistryModule = FModuleManager::LoadModuleChecked<FAssetRegistryModule>("AssetRegistry");
 
-		for ( auto PathIt = PathsToDelete.CreateConstIterator(); PathIt; ++PathIt )
+		for (const FString& PathToDelete : PathsToDelete)
 		{
-			AssetRegistryModule.Get().RemovePath(*PathIt);
+			FString PathToDeleteOnDisk;
+			if (FPackageName::TryConvertLongPackageNameToFilename(PathToDelete, PathToDeleteOnDisk) && IFileManager::Get().DeleteDirectory(*PathToDeleteOnDisk))
+			{
+				AssetRegistryModule.Get().RemovePath(PathToDelete);
+			}
 		}
 
 		return true;
diff --git a/Engine/Source/Editor/ContentBrowser/Private/SAssetView.cpp b/Engine/Source/Editor/ContentBrowser/Private/SAssetView.cpp
--- a/Engine/Source/Editor/ContentBrowser/Private/SAssetView.cpp
+++ b/Engine/Source/Editor/ContentBrowser/Private/SAssetView.cpp
@@ -2388,11 +2388,11 @@ void SAssetView::OnAssetLoaded(UObject* Asset)
 	}
 }
 
-void SAssetView::OnObjectPropertyChanged(UObject* Asset, FPropertyChangedEvent& PropertyChangedEvent)
+void SAssetView::OnObjectPropertyChanged(UObject* Object, FPropertyChangedEvent& PropertyChangedEvent)
 {
-	if ( Asset != NULL )
+	if (Object != nullptr && Object->IsAsset())
 	{
-		RecentlyLoadedOrChangedAssets.Add( FName(*Asset->GetPathName()), Asset );
+		RecentlyLoadedOrChangedAssets.Add( FName(*Object->GetPathName()), Object);
 	}
 }
 
diff --git a/Engine/Source/Editor/ContentBrowser/Private/SAssetView.h b/Engine/Source/Editor/ContentBrowser/Private/SAssetView.h
--- a/Engine/Source/Editor/ContentBrowser/Private/SAssetView.h
+++ b/Engine/Source/Editor/ContentBrowser/Private/SAssetView.h
@@ -336,7 +336,7 @@ private:
 	void OnAssetLoaded(UObject* Asset);
 
 	/** Handler for when an asset's property has changed */
-	void OnObjectPropertyChanged(UObject* Asset, FPropertyChangedEvent& PropertyChangedEvent);
+	void OnObjectPropertyChanged(UObject* Object, FPropertyChangedEvent& PropertyChangedEvent);
 
 	/** Called when the class hierarchy is updated due to the available modules changing */
 	void OnClassHierarchyUpdated();
diff --git a/Engine/Source/Editor/DetailCustomizations/DetailCustomizations.Build.cs b/Engine/Source/Editor/DetailCustomizations/DetailCustomizations.Build.cs
--- a/Engine/Source/Editor/DetailCustomizations/DetailCustomizations.Build.cs
+++ b/Engine/Source/Editor/DetailCustomizations/DetailCustomizations.Build.cs
@@ -47,6 +47,7 @@ public class DetailCustomizations : ModuleRules
 				"ConfigEditor",
                 "Persona",
                 "CinematicCamera",
+                "ComponentVisualizers"
 			}
 		);
 
@@ -58,6 +59,7 @@ public class DetailCustomizations : ModuleRules
 				"LandscapeEditor",
 				"PropertyEditor",
 				"GameProjectGeneration",
+                "ComponentVisualizers"
 			}
 		);
 
