Fix localization errors in these files

diff --git a/Engine/Source/Developer/SettingsEditor/Private/Widgets/SSettingsSectionHeader.cpp b/Engine/Source/Developer/SettingsEditor/Private/Widgets/SSettingsSectionHeader.cpp
--- a/Engine/Source/Developer/SettingsEditor/Private/Widgets/SSettingsSectionHeader.cpp
+++ b/Engine/Source/Developer/SettingsEditor/Private/Widgets/SSettingsSectionHeader.cpp
@@ -325,7 +325,7 @@ FReply SSettingsSectionHeader::HandleSetAsDefaultButtonClicked()
 					{
 						if (!CheckOutOrAddDefaultConfigFile())
 						{
-							if (FMessageDialog::Open(EAppMsgType::YesNo, LOCTEXT("SaveAsDefaultsIsReadOnlyMessage", "The source control operation failed. Would you like to make it writable?")) == EAppReturnType::Yes)
+							if (FMessageDialog::Open(EAppMsgType::YesNo, LOCTEXT("SaveAsDefaultsSourceControlOperationFailed", "The source control operation failed. Would you like to make it writable?")) == EAppReturnType::Yes)
 							{
 								MakeDefaultConfigFileWritable();
 							}
@@ -365,7 +365,7 @@ FReply SSettingsSectionHeader::HandleSetAsDefaultButtonClicked()
 			{
 				if (!CheckOutOrAddDefaultConfigFile(true))
 				{
-					FMessageDialog::Open(EAppMsgType::Ok, LOCTEXT("SaveAsDefaultsIsReadOnlyMessage", "The source control operation failed. You will need to add it manually"));
+					FMessageDialog::Open(EAppMsgType::Ok, LOCTEXT("SaveAsDefaultsSourceControlFailedAddManually", "The source control operation failed. You will need to add it manually"));
 					return FReply::Handled();
 				}
 			}
diff --git a/Engine/Source/Editor/KismetWidgets/Private/SPinTypeSelector.cpp b/Engine/Source/Editor/KismetWidgets/Private/SPinTypeSelector.cpp
--- a/Engine/Source/Editor/KismetWidgets/Private/SPinTypeSelector.cpp
+++ b/Engine/Source/Editor/KismetWidgets/Private/SPinTypeSelector.cpp
@@ -1011,10 +1011,10 @@ FText SPinTypeSelector::GetToolTipForComboBoxSecondaryType() const
 {
 	if (IsEnabled())
 	{
-		return LOCTEXT("PinTypeSelector", "Select the map's value type.");
+		return LOCTEXT("PinTypeValueSelector", "Select the map's value type.");
 	}
 
-	return LOCTEXT("PinTypeSelector_Disabled", "Cannot edit map value type when they are inherited from parent.");
+	return LOCTEXT("PinTypeSelector_ValueDisabled", "Cannot edit map value type when they are inherited from parent.");
 }
 
 FText SPinTypeSelector::GetToolTipForArrayWidget() const
diff --git a/Engine/Source/Editor/Persona/Private/PersonaMeshDetails.cpp b/Engine/Source/Editor/Persona/Private/PersonaMeshDetails.cpp
--- a/Engine/Source/Editor/Persona/Private/PersonaMeshDetails.cpp
+++ b/Engine/Source/Editor/Persona/Private/PersonaMeshDetails.cpp
@@ -511,7 +511,7 @@ void FPersonaMeshDetails::AddLODLevelCategories(IDetailLayoutBuilder& DetailLayo
 							SNew(SButton)
 							.ButtonStyle(FEditorStyle::Get(), "HoverHintOnly")
 							.Text(LOCTEXT("AddLODLevelCategories_MaterialArrayOpAdd", "Add Material Slot"))
-							.ToolTipText(LOCTEXT("AddLODLevelCategories_MaterialArrayOpAdd", "Add Material Slot at the end of the Material slot array. Those Material slots can be used to override a LODs section, (not the base LOD)"))
+							.ToolTipText(LOCTEXT("AddLODLevelCategories_MaterialArrayOpAdd_Tooltip", "Add Material Slot at the end of the Material slot array. Those Material slots can be used to override a LODs section, (not the base LOD)"))
 							.ContentPadding(4.0f)
 							.ForegroundColor(FSlateColor::UseForeground())
 							.OnClicked(this, &FPersonaMeshDetails::AddMaterialSlot)
diff --git a/Engine/Source/Editor/PropertyEditor/Private/UserInterface/PropertyEditor/SPropertyEditorMap.cpp b/Engine/Source/Editor/PropertyEditor/Private/UserInterface/PropertyEditor/SPropertyEditorMap.cpp
--- a/Engine/Source/Editor/PropertyEditor/Private/UserInterface/PropertyEditor/SPropertyEditorMap.cpp
+++ b/Engine/Source/Editor/PropertyEditor/Private/UserInterface/PropertyEditor/SPropertyEditorMap.cpp
@@ -61,4 +61,6 @@ FText SPropertyEditorMap::GetMapTooltipText() const
 bool SPropertyEditorMap::CanEdit() const
 {
 	return PropertyEditor.IsValid() ? !PropertyEditor->IsEditConst() : true;
-}
\ No newline at end of file
+}
+
+#undef LOCTEXT_NAMESPACE
\ No newline at end of file
diff --git a/Engine/Source/Editor/PropertyEditor/Private/UserInterface/PropertyEditor/SPropertyEditorMap.h b/Engine/Source/Editor/PropertyEditor/Private/UserInterface/PropertyEditor/SPropertyEditorMap.h
--- a/Engine/Source/Editor/PropertyEditor/Private/UserInterface/PropertyEditor/SPropertyEditorMap.h
+++ b/Engine/Source/Editor/PropertyEditor/Private/UserInterface/PropertyEditor/SPropertyEditorMap.h
@@ -4,8 +4,6 @@
 #include "PropertyEditor.h"
 #include "PropertyEditorConstants.h"
 
-#define LOCTEXT_NAMESPACE "PropertyEditor"
-
 class SPropertyEditorMap : public SCompoundWidget
 {
 public:
diff --git a/Engine/Source/Editor/PropertyEditor/Private/UserInterface/PropertyEditor/SPropertyEditorSet.cpp b/Engine/Source/Editor/PropertyEditor/Private/UserInterface/PropertyEditor/SPropertyEditorSet.cpp
--- a/Engine/Source/Editor/PropertyEditor/Private/UserInterface/PropertyEditor/SPropertyEditorSet.cpp
+++ b/Engine/Source/Editor/PropertyEditor/Private/UserInterface/PropertyEditor/SPropertyEditorSet.cpp
@@ -61,4 +61,6 @@ FText SPropertyEditorSet::GetSetTooltipText() const
 bool SPropertyEditorSet::CanEdit() const
 {
 	return PropertyEditor.IsValid() ? !PropertyEditor->IsEditConst() : true;
-}
\ No newline at end of file
+}
+
+#undef LOCTEXT_NAMESPACE
\ No newline at end of file
diff --git a/Engine/Source/Editor/PropertyEditor/Private/UserInterface/PropertyEditor/SPropertyEditorSet.h b/Engine/Source/Editor/PropertyEditor/Private/UserInterface/PropertyEditor/SPropertyEditorSet.h
--- a/Engine/Source/Editor/PropertyEditor/Private/UserInterface/PropertyEditor/SPropertyEditorSet.h
+++ b/Engine/Source/Editor/PropertyEditor/Private/UserInterface/PropertyEditor/SPropertyEditorSet.h
@@ -4,8 +4,6 @@
 #include "PropertyEditor.h"
 #include "PropertyEditorConstants.h"
 
-#define LOCTEXT_NAMESPACE "PropertyEditor"
-
 class SPropertyEditorSet : public SCompoundWidget
 {
 public:
diff --git a/Engine/Source/Editor/StaticMeshEditor/Private/StaticMeshEditorTools.cpp b/Engine/Source/Editor/StaticMeshEditor/Private/StaticMeshEditorTools.cpp
--- a/Engine/Source/Editor/StaticMeshEditor/Private/StaticMeshEditorTools.cpp
+++ b/Engine/Source/Editor/StaticMeshEditor/Private/StaticMeshEditorTools.cpp
@@ -1687,7 +1687,7 @@ void FMeshMaterialsLayout::AddToCategory(IDetailCategoryBuilder& CategoryBuilder
 					SNew(SButton)
 					.ButtonStyle(FEditorStyle::Get(), "HoverHintOnly")
 					.Text(LOCTEXT("AddLODLevelCategories_MaterialArrayOpAdd", "Add Material Slot"))
-					.ToolTipText(LOCTEXT("AddLODLevelCategories_MaterialArrayOpAdd", "Add Material Slot at the end of the Material slot array. Those Material slots can be used to override a LODs section, (not the base LOD)"))
+					.ToolTipText(LOCTEXT("AddLODLevelCategories_MaterialArrayOpAdd_Tooltip", "Add Material Slot at the end of the Material slot array. Those Material slots can be used to override a LODs section, (not the base LOD)"))
 					.ContentPadding(4.0f)
 					.ForegroundColor(FSlateColor::UseForeground())
 					.OnClicked(this, &FMeshMaterialsLayout::AddMaterialSlot)
