Fix for crash on reordering in Settings Editor. Added nullcheck for object being edited.

diff --git a/Engine/Source/Developer/SettingsEditor/Private/Widgets/SSettingsEditor.cpp b/Engine/Source/Developer/SettingsEditor/Private/Widgets/SSettingsEditor.cpp
--- a/Engine/Source/Developer/SettingsEditor/Private/Widgets/SSettingsEditor.cpp
+++ b/Engine/Source/Developer/SettingsEditor/Private/Widgets/SSettingsEditor.cpp
@@ -158,69 +158,71 @@ void SSettingsEditor::NotifyPostChange( const FPropertyChangedEvent& PropertyCha
 
 		// Note while there could be multiple objects in the details panel, only one is ever edited at once
 		const UObject* ObjectBeingEdited = PropertyChangedEvent.GetObjectBeingEdited(0);
-
-		// Get the section from the edited object.  We cannot use the selected section as multiple sections can be shown at once in the settings details panel.
-		ISettingsSectionPtr Section = Model->GetSectionFromSectionObject(ObjectBeingEdited);
-		if(Section.IsValid())
+		if (ObjectBeingEdited != nullptr)
 		{
-			FString RelativePath;
-			bool bIsSourceControlled = false;
-			bool bIsNewFile = false;
-
-			// Attempt to checkout the file automatically
-			if (ObjectBeingEdited->GetClass()->HasAnyClassFlags(CLASS_DefaultConfig))
+			// Get the section from the edited object.  We cannot use the selected section as multiple sections can be shown at once in the settings details panel.
+			ISettingsSectionPtr Section = Model->GetSectionFromSectionObject(ObjectBeingEdited);
+			if (Section.IsValid())
 			{
-				RelativePath = ObjectBeingEdited->GetDefaultConfigFilename();
-				bIsSourceControlled = true;
-			}
-			else if (ObjectBeingEdited->GetClass()->HasAnyClassFlags(CLASS_Config))
-			{
-				RelativePath = ObjectBeingEdited->GetClass()->GetConfigName();
-			}
+				FString RelativePath;
+				bool bIsSourceControlled = false;
+				bool bIsNewFile = false;
 
-			FString FullPath = FPaths::ConvertRelativePathToFull(RelativePath);
+				// Attempt to checkout the file automatically
+				if (ObjectBeingEdited->GetClass()->HasAnyClassFlags(CLASS_DefaultConfig))
+				{
+					RelativePath = ObjectBeingEdited->GetDefaultConfigFilename();
+					bIsSourceControlled = true;
+				}
+				else if (ObjectBeingEdited->GetClass()->HasAnyClassFlags(CLASS_Config))
+				{
+					RelativePath = ObjectBeingEdited->GetClass()->GetConfigName();
+				}
 
-			if (!FPlatformFileManager::Get().GetPlatformFile().FileExists(*FullPath))
-			{
-				bIsNewFile = true;
-			}
+				FString FullPath = FPaths::ConvertRelativePathToFull(RelativePath);
 
-			if (!bIsSourceControlled || !SettingsHelpers::CheckOutOrAddFile(FullPath))
-			{
-				SettingsHelpers::MakeWritable(FullPath);
-			}
-			
-			RecordPreferenceChangedAnalytics(Section, PropertyChangedEvent);
+				if (!FPlatformFileManager::Get().GetPlatformFile().FileExists(*FullPath))
+				{
+					bIsNewFile = true;
+				}
+
+				if (!bIsSourceControlled || !SettingsHelpers::CheckOutOrAddFile(FullPath))
+				{
+					SettingsHelpers::MakeWritable(FullPath);
+				}
 
-			// Determine if the Property is an Array or Array Element
-			bool bIsArrayOrArrayElement = PropertyThatChanged->GetActiveMemberNode()->GetValue()->IsA(UArrayProperty::StaticClass()) 
-				|| PropertyThatChanged->GetActiveMemberNode()->GetValue()->ArrayDim > 1
-				|| ((Outer != nullptr) && Outer->IsA(UArrayProperty::StaticClass()));
+				RecordPreferenceChangedAnalytics(Section, PropertyChangedEvent);
 
-			bool bIsSetOrSetElement = PropertyThatChanged->GetActiveMemberNode()->GetValue()->IsA(USetProperty::StaticClass())
-				|| ((Outer != nullptr) && Outer->IsA(USetProperty::StaticClass()));
+				// Determine if the Property is an Array or Array Element
+				bool bIsArrayOrArrayElement = PropertyThatChanged->GetActiveMemberNode()->GetValue()->IsA(UArrayProperty::StaticClass())
+					|| PropertyThatChanged->GetActiveMemberNode()->GetValue()->ArrayDim > 1
+					|| ((Outer != nullptr) && Outer->IsA(UArrayProperty::StaticClass()));
 
-			bool bIsMapOrMapElement = PropertyThatChanged->GetActiveMemberNode()->GetValue()->IsA(UMapProperty::StaticClass())
-				|| ((Outer != nullptr) && Outer->IsA(UMapProperty::StaticClass()));
+				bool bIsSetOrSetElement = PropertyThatChanged->GetActiveMemberNode()->GetValue()->IsA(USetProperty::StaticClass())
+					|| ((Outer != nullptr) && Outer->IsA(USetProperty::StaticClass()));
 
-			if (Section->GetSettingsObject()->GetClass()->HasAnyClassFlags(CLASS_DefaultConfig) && !bIsArrayOrArrayElement && !bIsSetOrSetElement && !bIsMapOrMapElement)
-			{
-				Section->GetSettingsObject()->UpdateSinglePropertyInConfigFile(PropertyThatChanged->GetActiveMemberNode()->GetValue(), Section->GetSettingsObject()->GetDefaultConfigFilename());
-			}
-			else
-			{
-				Section->Save();
-			}
+				bool bIsMapOrMapElement = PropertyThatChanged->GetActiveMemberNode()->GetValue()->IsA(UMapProperty::StaticClass())
+					|| ((Outer != nullptr) && Outer->IsA(UMapProperty::StaticClass()));
 
-			if (bIsNewFile && bIsSourceControlled)
-			{
-				SettingsHelpers::CheckOutOrAddFile(FullPath);
-			}
+				if (Section->GetSettingsObject()->GetClass()->HasAnyClassFlags(CLASS_DefaultConfig) && !bIsArrayOrArrayElement && !bIsSetOrSetElement && !bIsMapOrMapElement)
+				{
+					Section->GetSettingsObject()->UpdateSinglePropertyInConfigFile(PropertyThatChanged->GetActiveMemberNode()->GetValue(), Section->GetSettingsObject()->GetDefaultConfigFilename());
+				}
+				else
+				{
+					Section->Save();
+				}
 
-			static const FName ConfigRestartRequiredKey = "ConfigRestartRequired";
-			if (PropertyChangedEvent.Property->GetBoolMetaData(ConfigRestartRequiredKey))
-			{
-				OnApplicationRestartRequiredDelegate.ExecuteIfBound();
+				if (bIsNewFile && bIsSourceControlled)
+				{
+					SettingsHelpers::CheckOutOrAddFile(FullPath);
+				}
+
+				static const FName ConfigRestartRequiredKey = "ConfigRestartRequired";
+				if (PropertyChangedEvent.Property->GetBoolMetaData(ConfigRestartRequiredKey))
+				{
+					OnApplicationRestartRequiredDelegate.ExecuteIfBound();
+				}
 			}
 		}
 	}
