Single property setting changes will now also call OnModified delegate for their section

diff --git a/Engine/Source/Developer/Settings/Private/SettingsSection.cpp b/Engine/Source/Developer/Settings/Private/SettingsSection.cpp
--- a/Engine/Source/Developer/Settings/Private/SettingsSection.cpp
+++ b/Engine/Source/Developer/Settings/Private/SettingsSection.cpp
@@ -190,6 +190,17 @@ bool FSettingsSection::ResetDefaults()
 	return false;
 }
 
+bool FSettingsSection::NotifySectionOnPropertyModified()
+{
+	bool bShouldSaveChanges = true;
+	// Notify the section that it has been modified
+	if (ModifiedDelegate.IsBound())
+	{
+		// return value of FOnModified indicates whether the modifications should be saved.
+		bShouldSaveChanges = ModifiedDelegate.Execute();
+	}
+	return bShouldSaveChanges;
+}
 
 bool FSettingsSection::Save()
 {
diff --git a/Engine/Source/Developer/Settings/Private/SettingsSection.h b/Engine/Source/Developer/Settings/Private/SettingsSection.h
--- a/Engine/Source/Developer/Settings/Private/SettingsSection.h
+++ b/Engine/Source/Developer/Settings/Private/SettingsSection.h
@@ -101,6 +101,7 @@ public:
 	}
 
 	virtual bool ResetDefaults() override;
+	virtual bool NotifySectionOnPropertyModified() override;
 	virtual bool Save() override;
 	virtual bool SaveDefaults() override;
 
diff --git a/Engine/Source/Developer/Settings/Public/ISettingsSection.h b/Engine/Source/Developer/Settings/Public/ISettingsSection.h
--- a/Engine/Source/Developer/Settings/Public/ISettingsSection.h
+++ b/Engine/Source/Developer/Settings/Public/ISettingsSection.h
@@ -143,6 +143,13 @@ public:
 	 */
 	virtual bool ResetDefaults() = 0;
 
+
+	/** Notifies a section that a single property has been modified 
+	*
+ 	* @return true if settings should be saved, false otherwise.
+	*/
+	virtual bool NotifySectionOnPropertyModified() = 0;
+
 	/**
 	 * Saves the settings in this section.
 	 *
diff --git a/Engine/Source/Developer/SettingsEditor/Private/Widgets/SSettingsEditor.cpp b/Engine/Source/Developer/SettingsEditor/Private/Widgets/SSettingsEditor.cpp
--- a/Engine/Source/Developer/SettingsEditor/Private/Widgets/SSettingsEditor.cpp
+++ b/Engine/Source/Developer/SettingsEditor/Private/Widgets/SSettingsEditor.cpp
@@ -211,7 +211,10 @@ void SSettingsEditor::NotifyPostChange( const FPropertyChangedEvent& PropertyCha
 
 				if (ObjectBeingEdited->GetClass()->HasAnyClassFlags(CLASS_DefaultConfig) && !bIsArrayOrArrayElement && !bIsSetOrSetElement && !bIsMapOrMapElement)
 				{
-					ObjectBeingEdited->UpdateSinglePropertyInConfigFile(PropertyThatChanged->GetActiveMemberNode()->GetValue(), ObjectBeingEdited->GetDefaultConfigFilename());
+					if(Section->NotifySectionOnPropertyModified())
+					{
+						ObjectBeingEdited->UpdateSinglePropertyInConfigFile(PropertyThatChanged->GetActiveMemberNode()->GetValue(), ObjectBeingEdited->GetDefaultConfigFilename());
+					}
 				}
 				else if (Section.IsValid())
 				{
