1. This adds a small triangle icon close to each one of the main sections of the Editor Preferences menu, so they can be quickly collapsed/expanded.
2. For efficiency purposes, it also pass some TSharedRef arguments by reference rather than by copy.

diff --git a/Engine/Source/Developer/SettingsEditor/Private/Widgets/SSettingsSectionHeader.cpp b/Engine/Source/Developer/SettingsEditor/Private/Widgets/SSettingsSectionHeader.cpp
--- a/Engine/Source/Developer/SettingsEditor/Private/Widgets/SSettingsSectionHeader.cpp
+++ b/Engine/Source/Developer/SettingsEditor/Private/Widgets/SSettingsSectionHeader.cpp
@@ -10,6 +10,7 @@
 #include "Framework/Application/SlateApplication.h"
 #include "Widgets/Input/SButton.h"
 #include "Widgets/Text/STextBlock.h"
+#include "Widgets/Views/STableRow.h"
 #include "EditorStyleSet.h"
 #include "ISourceControlModule.h"
 #include "SSettingsEditorCheckoutNotice.h"
@@ -18,7 +19,98 @@
 
 #define LOCTEXT_NAMESPACE "SSettingsEditor"
 
-void SSettingsSectionHeader::Construct(const FArguments& InArgs, const UObject* InSettingsObject, ISettingsEditorModelPtr InModel, TSharedPtr<IDetailsView> InDetailsView)
+#define S_SETTINGS_SECTION_HEADER_CONSTRUCT_CHILD_SLOT_INIT \
+	ChildSlot \
+	.Padding(FMargin(0.0f, 8.0f, 0.0f, 5.0f)) \
+	[ \
+		SNew(SVerticalBox) \
+		+ SVerticalBox::Slot() \
+		.AutoHeight() \
+		[ \
+			SNew(SHorizontalBox) \
+			+ SHorizontalBox::Slot() \
+			.FillWidth(1.0f) \
+			[ \
+				SNew(SVerticalBox) \
+				+ SVerticalBox::Slot() \
+				.AutoHeight() \
+				[
+
+#define S_SETTINGS_SECTION_HEADER_CONSTRUCT_CHILD_SLOT_END \
+				] \
+				+ SVerticalBox::Slot() \
+				.AutoHeight() \
+				.Padding(0.0f, 8.0f, 0.0f, 0.0f) \
+				[ \
+					/* category description */ \
+					SNew(STextBlock) \
+					.ColorAndOpacity(FSlateColor::UseSubduedForeground()) \
+					.Text(GetSettingsBoxDescriptionText()) \
+				] \
+			] \
+			+ SHorizontalBox::Slot() \
+			.AutoWidth() \
+			.HAlign(HAlign_Right) \
+			.VAlign(VAlign_Bottom) \
+			.Padding(16.0f, 0.0f, 0.0f, 0.0f) \
+			[ \
+				SNew(SHorizontalBox) \
+				.Visibility(this, &SSettingsSectionHeader::HandleButtonRowVisibility) \
+				+SHorizontalBox::Slot() \
+				[ \
+					/* set as default button */ \
+					SNew(SButton) \
+					.Visibility(this, &SSettingsSectionHeader::HandleSetAsDefaultButtonVisibility) \
+					.IsEnabled(this, &SSettingsSectionHeader::HandleSetAsDefaultButtonEnabled) \
+					.OnClicked(this, &SSettingsSectionHeader::HandleSetAsDefaultButtonClicked) \
+					.Text(LOCTEXT("SaveDefaultsButtonText", "Set as Default")) \
+					.ToolTipText(LOCTEXT("SaveDefaultsButtonTooltip", "Save the values below as the new default settings")) \
+				] \
+				+ SHorizontalBox::Slot() \
+				.AutoWidth() \
+				.Padding(8.0f, 0.0f, 0.0f, 0.0f) \
+				[ \
+					/* export button */ \
+					SNew(SButton) \
+					.IsEnabled(this, &SSettingsSectionHeader::HandleExportButtonEnabled) \
+					.OnClicked(this, &SSettingsSectionHeader::HandleExportButtonClicked) \
+					.Text(LOCTEXT("ExportButtonText", "Export...")) \
+					.ToolTipText(LOCTEXT("ExportButtonTooltip", "Export these settings to a file on your computer")) \
+				] \
+				+ SHorizontalBox::Slot() \
+				.AutoWidth() \
+				.Padding(8.0f, 0.0f, 0.0f, 0.0f) \
+				[ \
+					/* import button */ \
+					SNew(SButton) \
+					.IsEnabled(this, &SSettingsSectionHeader::HandleImportButtonEnabled) \
+					.OnClicked(this, &SSettingsSectionHeader::HandleImportButtonClicked) \
+					.Text(LOCTEXT("ImportButtonText", "Import...")) \
+					.ToolTipText(LOCTEXT("ImportButtonTooltip", "Import these settings from a file on your computer")) \
+				] \
+				+ SHorizontalBox::Slot() \
+				.AutoWidth() \
+				.Padding(8.0f, 0.0f, 0.0f, 0.0f) \
+				[ \
+					/* reset defaults button */ \
+					SNew(SButton) \
+					.Visibility(this, &SSettingsSectionHeader::HandleSetAsDefaultButtonVisibility) \
+					.IsEnabled(this, &SSettingsSectionHeader::HandleResetToDefaultsButtonEnabled) \
+					.OnClicked(this, &SSettingsSectionHeader::HandleResetDefaultsButtonClicked) \
+					.Text(LOCTEXT("ResetDefaultsButtonText", "Reset to Defaults")) \
+					.ToolTipText(LOCTEXT("ResetDefaultsButtonTooltip", "Reset the settings below to their default values")) \
+				] \
+			] \
+		] \
+		+ SVerticalBox::Slot() \
+		.AutoHeight() \
+		.Padding(0.0f, 8.0f, 0.0f, 0.0f) \
+		[ \
+			FileWatcherWidget.ToSharedRef() \
+		] \
+	];
+
+void SSettingsSectionHeader::CommonConstruct(const UObject* InSettingsObject, ISettingsEditorModelPtr InModel, TSharedPtr<IDetailsView> InDetailsView)
 {
 	Model = InModel;
 	SettingsObject = MakeWeakObjectPtr(const_cast<UObject*>(InSettingsObject));
@@ -33,100 +125,42 @@ void SSettingsSectionHeader::Construct(const FArguments& InArgs, const UObject*
 		.Visibility(this, &SSettingsSectionHeader::HandleCheckoutNoticeVisibility)
 		.OnFileProbablyModifiedExternally(this, &SSettingsSectionHeader::HandleCheckoutNoticeFileProbablyModifiedExternally)
 		.ConfigFilePath(this, &SSettingsSectionHeader::HandleCheckoutNoticeConfigFilePath);
+}
 
+void SSettingsSectionHeader::Construct(const FArguments& InArgs, const UObject* InSettingsObject, ISettingsEditorModelPtr InModel, TSharedPtr<IDetailsView> InDetailsView)
+{
+	CommonConstruct(InSettingsObject, InModel, InDetailsView);
 
-	ChildSlot
-	.Padding(FMargin(0.0f, 8.0f, 0.0f, 5.0f))
-	[
-		SNew(SVerticalBox)
-		+ SVerticalBox::Slot()
-		.AutoHeight()
-		[
-			SNew(SHorizontalBox)
-			+ SHorizontalBox::Slot()
-			.FillWidth(1.0f)
-			[
-				SNew(SVerticalBox)
-				+ SVerticalBox::Slot()
-				.AutoHeight()
-				[
+	S_SETTINGS_SECTION_HEADER_CONSTRUCT_CHILD_SLOT_INIT
 					// category title
 					SNew(STextBlock)
 					.Font(FEditorStyle::GetFontStyle("SettingsEditor.CatgoryAndSectionFont"))
 					.Text(GetSettingsBoxTitleText())
-				]
-				+ SVerticalBox::Slot()
-				.AutoHeight()
-				.Padding(0.0f, 8.0f, 0.0f, 0.0f)
-				[
-					// category description
-					SNew(STextBlock)
-					.ColorAndOpacity(FSlateColor::UseSubduedForeground())
-					.Text(GetSettingsBoxDescriptionText())
-				]
-			]
-			+ SHorizontalBox::Slot()
-			.AutoWidth()
-			.HAlign(HAlign_Right)
-			.VAlign(VAlign_Bottom)
-			.Padding(16.0f, 0.0f, 0.0f, 0.0f)
-			[
-				SNew(SHorizontalBox)
-				.Visibility(this, &SSettingsSectionHeader::HandleButtonRowVisibility)
-				+SHorizontalBox::Slot()
-				[
-					// set as default button
-					SNew(SButton)
-					.Visibility(this, &SSettingsSectionHeader::HandleSetAsDefaultButtonVisibility)
-					.IsEnabled(this, &SSettingsSectionHeader::HandleSetAsDefaultButtonEnabled)
-					.OnClicked(this, &SSettingsSectionHeader::HandleSetAsDefaultButtonClicked)
-					.Text(LOCTEXT("SaveDefaultsButtonText", "Set as Default"))
-					.ToolTipText(LOCTEXT("SaveDefaultsButtonTooltip", "Save the values below as the new default settings"))
-				]
-				+ SHorizontalBox::Slot()
-				.AutoWidth()
-				.Padding(8.0f, 0.0f, 0.0f, 0.0f)
-				[
-					// export button
-					SNew(SButton)
-					.IsEnabled(this, &SSettingsSectionHeader::HandleExportButtonEnabled)
-					.OnClicked(this, &SSettingsSectionHeader::HandleExportButtonClicked)
-					.Text(LOCTEXT("ExportButtonText", "Export..."))
-					.ToolTipText(LOCTEXT("ExportButtonTooltip", "Export these settings to a file on your computer"))
-				]
-				+ SHorizontalBox::Slot()
-				.AutoWidth()
-				.Padding(8.0f, 0.0f, 0.0f, 0.0f)
-				[
-					// import button
-					SNew(SButton)
-					.IsEnabled(this, &SSettingsSectionHeader::HandleImportButtonEnabled)
-					.OnClicked(this, &SSettingsSectionHeader::HandleImportButtonClicked)
-					.Text(LOCTEXT("ImportButtonText", "Import..."))
-					.ToolTipText(LOCTEXT("ImportButtonTooltip", "Import these settings from a file on your computer"))
-				]
-				+ SHorizontalBox::Slot()
-				.AutoWidth()
-				.Padding(8.0f, 0.0f, 0.0f, 0.0f)
-				[
-					// reset defaults button
-					SNew(SButton)
-					.Visibility(this, &SSettingsSectionHeader::HandleSetAsDefaultButtonVisibility)
-					.IsEnabled(this, &SSettingsSectionHeader::HandleResetToDefaultsButtonEnabled)
-					.OnClicked(this, &SSettingsSectionHeader::HandleResetDefaultsButtonClicked)
-					.Text(LOCTEXT("ResetDefaultsButtonText", "Reset to Defaults"))
-					.ToolTipText(LOCTEXT("ResetDefaultsButtonTooltip", "Reset the settings below to their default values"))
-				]
-			]
-		]
-		+ SVerticalBox::Slot()
-		.AutoHeight()
-		.Padding(0.0f, 8.0f, 0.0f, 0.0f)
-		[
-			FileWatcherWidget.ToSharedRef()
-		]
-	];
+	S_SETTINGS_SECTION_HEADER_CONSTRUCT_CHILD_SLOT_END
+}
 
+void SSettingsSectionHeader::Construct(const FArguments& InArgs, const UObject* InSettingsObject, ISettingsEditorModelPtr InModel, TSharedPtr<IDetailsView> InDetailsView, const TSharedRef<ITableRow>& InTableRow)
+{
+	CommonConstruct(InSettingsObject, InModel, InDetailsView);
+
+	S_SETTINGS_SECTION_HEADER_CONSTRUCT_CHILD_SLOT_INIT
+					SNew(SHorizontalBox)
+					+ SHorizontalBox::Slot()
+					.VAlign(VAlign_Center)
+					.HAlign(HAlign_Left)
+					.AutoWidth()
+					[
+						SNew(SExpanderArrow, InTableRow)
+						.Visibility(true ? EVisibility::Visible : EVisibility::Collapsed)
+					]
+					+ SHorizontalBox::Slot()
+					[
+					// category title
+					SNew(STextBlock)
+					.Font(FEditorStyle::GetFontStyle("SettingsEditor.CatgoryAndSectionFont"))
+					.Text(GetSettingsBoxTitleText())
+					]
+	S_SETTINGS_SECTION_HEADER_CONSTRUCT_CHILD_SLOT_END
 }
 
 FText SSettingsSectionHeader::GetSettingsBoxTitleText() const
@@ -496,7 +530,7 @@ void SSettingsSectionHeader::OnSettingsSelectionChanged()
 }
 
 
-FSettingsDetailRootObjectCustomization::FSettingsDetailRootObjectCustomization(ISettingsEditorModelPtr InModel, TSharedRef<IDetailsView> InDetailsView)
+FSettingsDetailRootObjectCustomization::FSettingsDetailRootObjectCustomization(ISettingsEditorModelPtr InModel, const TSharedRef<IDetailsView>& InDetailsView)
 	: Model(InModel)
 	, DetailsView(InDetailsView)
 {
@@ -514,6 +548,11 @@ TSharedPtr<SWidget> FSettingsDetailRootObjectCustomization::CustomizeObjectHeade
 	return SNew(SSettingsSectionHeader, InRootObject, Model, DetailsView.Pin());
 }
 
+TSharedPtr<SWidget> FSettingsDetailRootObjectCustomization::CustomizeObjectHeader(const UObject* InRootObject, const TSharedRef<ITableRow>& InTableRow)
+{
+	return SNew(SSettingsSectionHeader, InRootObject, Model, DetailsView.Pin(), InTableRow);
+}
+
 bool FSettingsDetailRootObjectCustomization::IsObjectVisible(const UObject* InRootObject) const
 {
 	return SelectedSectionObject == nullptr || SelectedSectionObject == InRootObject || DetailsView.Pin()->HasActiveSearch();
diff --git a/Engine/Source/Developer/SettingsEditor/Private/Widgets/SSettingsSectionHeader.h b/Engine/Source/Developer/SettingsEditor/Private/Widgets/SSettingsSectionHeader.h
--- a/Engine/Source/Developer/SettingsEditor/Private/Widgets/SSettingsSectionHeader.h
+++ b/Engine/Source/Developer/SettingsEditor/Private/Widgets/SSettingsSectionHeader.h
@@ -23,8 +23,12 @@ public:
 	SLATE_END_ARGS()
 
 	void Construct(const FArguments& InArgs, const UObject* InSettingsObject, ISettingsEditorModelPtr InModel, TSharedPtr<IDetailsView> InDetailsView);
+	void Construct(const FArguments& InArgs, const UObject* InSettingsObject, ISettingsEditorModelPtr InModel, TSharedPtr<IDetailsView> InDetailsView, const TSharedRef<ITableRow>& InTableRow);
+
 private:
 
+	void CommonConstruct(const UObject* InSettingsObject, ISettingsEditorModelPtr InModel, TSharedPtr<IDetailsView> InDetailsView);
+
 	FText GetSettingsBoxTitleText() const;
 	FText GetSettingsBoxDescriptionText() const;
 
@@ -107,12 +111,13 @@ private:
 class FSettingsDetailRootObjectCustomization : public IDetailRootObjectCustomization
 {
 public:
-	FSettingsDetailRootObjectCustomization(ISettingsEditorModelPtr InModel, TSharedRef<IDetailsView> InDetailsView);
+	FSettingsDetailRootObjectCustomization(ISettingsEditorModelPtr InModel, const TSharedRef<IDetailsView>& InDetailsView);
 
 	void Initialize();
 
 	/** IDetailRootObjectCustomization interface */
 	virtual TSharedPtr<SWidget> CustomizeObjectHeader(const UObject* InRootObject) override;
+	virtual TSharedPtr<SWidget> CustomizeObjectHeader(const UObject* InRootObject, const TSharedRef<ITableRow>& InTableRow) override;
 	virtual bool IsObjectVisible(const UObject* InRootObject) const override;
 	virtual bool ShouldDisplayHeader(const UObject* InRootObject) const override;
 private:
diff --git a/Engine/Source/Editor/PropertyEditor/Private/DetailMultiTopLevelObjectRootNode.cpp b/Engine/Source/Editor/PropertyEditor/Private/DetailMultiTopLevelObjectRootNode.cpp
--- a/Engine/Source/Editor/PropertyEditor/Private/DetailMultiTopLevelObjectRootNode.cpp
+++ b/Engine/Source/Editor/PropertyEditor/Private/DetailMultiTopLevelObjectRootNode.cpp
@@ -4,29 +4,28 @@
 #include "IDetailRootObjectCustomization.h"
 #include "DetailWidgetRow.h"
 
-void SDetailMultiTopLevelObjectTableRow::Construct( const FArguments& InArgs, TSharedRef<FDetailTreeNode> InOwnerTreeNode, const TSharedRef<SWidget>& InCustomizedWidgetContents, const TSharedRef<STableViewBase>& InOwnerTableView )
+void SDetailMultiTopLevelObjectTableRow::Construct(const FArguments& InArgs, const TSharedRef<FDetailTreeNode>& InOwnerTreeNode, const TSharedRef<SWidget>& InCustomizedWidgetContents, const TSharedRef<STableViewBase>& InOwnerTableView)
+{
+	Construct(InArgs, InOwnerTreeNode);
+	ChildSlotConstruct(InCustomizedWidgetContents, InOwnerTableView);
+	// Disable toggle functionality
+	bShowExpansionArrow = false;
+}
+
+void SDetailMultiTopLevelObjectTableRow::Construct(const FArguments& InArgs, const TSharedRef<FDetailTreeNode>& InOwnerTreeNode)
 {
 	OwnerTreeNode = InOwnerTreeNode;
 	bShowExpansionArrow = InArgs._ShowExpansionArrow;
+}
 
+void SDetailMultiTopLevelObjectTableRow::ChildSlotConstruct(const TSharedRef<SWidget>& InCustomizedWidgetContents, const TSharedRef<STableViewBase>& InOwnerTableView)
+{
 	ChildSlot
-	[	
+	[
 		SNew( SBox )
 		.Padding( FMargin( 0.0f, 0.0f, SDetailTableRowBase::ScrollbarPaddingSize, 0.0f ) )
 		[
-			SNew( SHorizontalBox )
-			+SHorizontalBox::Slot()
-			.VAlign(VAlign_Center)
-			.Padding(2.0f, 2.0f, 2.0f, 2.0f)
-			.AutoWidth()
-			[
-				SNew( SExpanderArrow, SharedThis(this) )
-				.Visibility(bShowExpansionArrow ? EVisibility::Visible : EVisibility::Collapsed)
-			]
-			+SHorizontalBox::Slot()
-			[
-				InCustomizedWidgetContents
-			]
+			InCustomizedWidgetContents
 		]
 	];
 
@@ -36,6 +35,8 @@ void SDetailMultiTopLevelObjectTableRow::Construct( const FArguments& InArgs, TS
 			.ShowSelection(false),
 		InOwnerTableView
 	);
+	// Enable toggle functionality
+	bShowExpansionArrow = true;
 }
 
 
@@ -98,20 +99,15 @@ ENodeVisibility FDetailMultiTopLevelObjectRootNode::GetVisibility() const
 TSharedRef< ITableRow > FDetailMultiTopLevelObjectRootNode::GenerateWidgetForTableView(const TSharedRef<STableViewBase>& OwnerTable, const FDetailColumnSizeData& ColumnSizeData, bool bAllowFavoriteSystem)
 {
 	FDetailWidgetRow Row;
-	GenerateStandaloneWidget(Row);
-
-	return SNew(SDetailMultiTopLevelObjectTableRow, AsShared(), Row.NameWidget.Widget, OwnerTable);
+	TSharedRef<SDetailMultiTopLevelObjectTableRow> DetailMultiTopLevelObjectTableRow = SNew(SDetailMultiTopLevelObjectTableRow, AsShared());
+	GenerateStandaloneWidget(Row, DetailMultiTopLevelObjectTableRow);
+	DetailMultiTopLevelObjectTableRow->ChildSlotConstruct(Row.NameWidget.Widget, OwnerTable);
+	return DetailMultiTopLevelObjectTableRow;
 }
 
 
-bool FDetailMultiTopLevelObjectRootNode::GenerateStandaloneWidget(FDetailWidgetRow& OutRow) const
+bool GenerateStandaloneWidgetInternal(FDetailWidgetRow& OutRow, TSharedPtr<SWidget>& HeaderWidget, const FName& NodeName)
 {
-	TSharedPtr<SWidget> HeaderWidget;
-	if (RootObjectCustomization.IsValid() && RootObject.IsValid())
-	{
-		HeaderWidget = RootObjectCustomization.Pin()->CustomizeObjectHeader(RootObject.Get());
-	}
-
 	if (!HeaderWidget.IsValid())
 	{
 		// no customization was supplied or was passed back from the interface as invalid
@@ -130,6 +126,29 @@ bool FDetailMultiTopLevelObjectRootNode::GenerateStandaloneWidget(FDetailWidgetR
 	return true;
 }
 
+bool FDetailMultiTopLevelObjectRootNode::GenerateStandaloneWidget(FDetailWidgetRow& OutRow) const
+{
+	TSharedPtr<SWidget> HeaderWidget;
+	if (RootObjectCustomization.IsValid() && RootObject.IsValid())
+	{
+		HeaderWidget = RootObjectCustomization.Pin()->CustomizeObjectHeader(RootObject.Get());
+	}
+
+	return GenerateStandaloneWidgetInternal(OutRow, HeaderWidget, NodeName);
+}
+
+
+bool FDetailMultiTopLevelObjectRootNode::GenerateStandaloneWidget(FDetailWidgetRow& OutRow, const TSharedRef<ITableRow>& InTableRow) const
+{
+	TSharedPtr<SWidget> HeaderWidget;
+	if (RootObjectCustomization.IsValid() && RootObject.IsValid())
+	{
+		HeaderWidget = RootObjectCustomization.Pin()->CustomizeObjectHeader(RootObject.Get(), InTableRow);
+	}
+
+	return GenerateStandaloneWidgetInternal(OutRow, HeaderWidget, NodeName);
+}
+
 
 void FDetailMultiTopLevelObjectRootNode::GetChildren(FDetailNodeList& OutChildren )
 {
diff --git a/Engine/Source/Editor/PropertyEditor/Private/DetailMultiTopLevelObjectRootNode.h b/Engine/Source/Editor/PropertyEditor/Private/DetailMultiTopLevelObjectRootNode.h
--- a/Engine/Source/Editor/PropertyEditor/Private/DetailMultiTopLevelObjectRootNode.h
+++ b/Engine/Source/Editor/PropertyEditor/Private/DetailMultiTopLevelObjectRootNode.h
@@ -26,7 +26,17 @@ public:
 		SLATE_ARGUMENT( bool, ShowExpansionArrow )
 	SLATE_END_ARGS()
 
-	void Construct( const FArguments& InArgs, TSharedRef<FDetailTreeNode> InOwnerTreeNode, const TSharedRef<SWidget>& InCustomizedWidgetContents, const TSharedRef<STableViewBase>& InOwnerTableView );
+	/**
+	* Use this construct when the toggle functionality (with the SExpanderArrow) is not desired.
+	* GenerateStandaloneWidget(Row) must be called before this function.
+	*/
+	void Construct( const FArguments& InArgs, const TSharedRef<FDetailTreeNode>& InOwnerTreeNode, const TSharedRef<SWidget>& InCustomizedWidgetContents, const TSharedRef<STableViewBase>& InOwnerTableView );
+	/**
+	* Use this construct together with ChildSlotConstruct when the toggle functionality (with the SExpanderArrow) is desired.
+	* GenerateStandaloneWidget(Row, DetailMultiTopLevelObjectTableRow) and ChildSlotConstruct(Row.NameWidget.Widget, OwnerTable) must be called after this function.
+	*/
+	void Construct( const FArguments& InArgs, const TSharedRef<FDetailTreeNode>& InOwnerTreeNode );
+	void ChildSlotConstruct(const TSharedRef<SWidget>& InCustomizedWidgetContents, const TSharedRef<STableViewBase>& InOwnerTableView);
 private:
 	const FSlateBrush* GetBackgroundImage() const;
 private:
@@ -48,7 +58,14 @@ private:
 	virtual bool ShouldBeExpanded() const override { return true; }
 	virtual ENodeVisibility GetVisibility() const override;
 	virtual TSharedRef< ITableRow > GenerateWidgetForTableView( const TSharedRef<STableViewBase>& OwnerTable, const FDetailColumnSizeData& ColumnSizeData, bool bAllowFavoriteSystem) override;
+	/**
+	* Use this version of GenerateStandaloneWidget when the toggle functionality (with the SExpanderArrow) is not used.
+	*/
 	virtual bool GenerateStandaloneWidget(FDetailWidgetRow& OutRow) const override;
+	/**
+	* Use this version of GenerateStandaloneWidget when the toggle functionality (with the SExpanderArrow) is being used.
+	*/
+	virtual bool GenerateStandaloneWidget(FDetailWidgetRow& OutRow, const TSharedRef<ITableRow>& InTableRow) const;
 	virtual void GetChildren(FDetailNodeList& OutChildren )  override;
 	virtual void FilterNode( const FDetailFilter& InFilter ) override;
 	virtual void Tick( float DeltaTime ) override {}
diff --git a/Engine/Source/Editor/PropertyEditor/Public/IDetailRootObjectCustomization.h b/Engine/Source/Editor/PropertyEditor/Public/IDetailRootObjectCustomization.h
--- a/Engine/Source/Editor/PropertyEditor/Public/IDetailRootObjectCustomization.h
+++ b/Engine/Source/Editor/PropertyEditor/Public/IDetailRootObjectCustomization.h
@@ -4,6 +4,7 @@
 #include "CoreMinimal.h"
 
 class SWidget;
+class ITableRow;
 
 /** 
  * Interface for any class that lays out details for a specific class
@@ -20,6 +21,17 @@ public:
  	 */
 	virtual TSharedPtr<SWidget> CustomizeObjectHeader(const UObject* InRootObject) = 0;
 
+	/** 
+	 * Called when the details panel wants to display an object header widget for a given object
+	 *
+	 * @param	InRootObject	The object whose header is being customized
+	 * @param	InTableRow		The ITableRow object (table views to talk to their rows) that will use the current IDetailRootObjectCustomization element
+ 	 */
+	virtual TSharedPtr<SWidget> CustomizeObjectHeader(const UObject* InRootObject, const TSharedRef<ITableRow>& InTableRow)
+	{
+		return CustomizeObjectHeader(InRootObject);
+	}
+
 	/**
 	 * Whether or not the object and all of its children should be visible in the details panel
 	 */
