PR #4284: Enable EditorSettings import button (Contributed by [REDACTED_PERSON])

diff --git a/Engine/Source/Developer/SettingsEditor/Private/Widgets/SSettingsSectionHeader.cpp b/Engine/Source/Developer/SettingsEditor/Private/Widgets/SSettingsSectionHeader.cpp
--- a/Engine/Source/Developer/SettingsEditor/Private/Widgets/SSettingsSectionHeader.cpp
+++ b/Engine/Source/Developer/SettingsEditor/Private/Widgets/SSettingsSectionHeader.cpp
@@ -288,7 +288,7 @@ FString SSettingsSectionHeader::GetDefaultConfigFilePath() const
 */
 bool SSettingsSectionHeader::IsDefaultConfigCheckOutNeeded(bool bForceSourceControlUpdate) const
 {
-	if(SettingsObject.IsValid() && SettingsObject->GetClass()->HasAnyClassFlags(CLASS_Config | CLASS_DefaultConfig))
+	if(SettingsObject.IsValid() && SettingsObject->GetClass()->HasAllClassFlags(CLASS_Config | CLASS_DefaultConfig))
 	{
 		// We can only fetch the file watcher if it's visible otherwise fallback to source control
 		if (FileWatcherWidget->GetVisibility().IsVisible())
@@ -329,7 +329,7 @@ bool SSettingsSectionHeader::HandleResetToDefaultsButtonEnabled() const
 
 EVisibility SSettingsSectionHeader::HandleSetAsDefaultButtonVisibility() const
 {
-	return (SettingsSection.IsValid() && SettingsSection->HasDefaultSettingsObject() && SettingsSection->CanSaveDefaults()) ? EVisibility::Collapsed : EVisibility::Visible;
+	return (SettingsSection.IsValid() && SettingsSection->CanSaveDefaults()) ? EVisibility::Visible : EVisibility::Collapsed;
 }
 
 FReply SSettingsSectionHeader::HandleSetAsDefaultButtonClicked()
@@ -515,7 +515,7 @@ void SSettingsSectionHeader::HandleCheckoutNoticeFileProbablyModifiedExternally(
 /** Callback for determining the visibility of the 'Locked' notice. */
 EVisibility SSettingsSectionHeader::HandleCheckoutNoticeVisibility() const
 {
-	// Only Defaultfig are under source control, so the checkout notice should not be visible for the other cases
+	// Only DefaultConfig are under source control, so the checkout notice should not be visible for the other cases
 	if(SettingsObject.IsValid() && SettingsObject->GetClass()->HasAnyClassFlags(CLASS_DefaultConfig) && !DetailsView.Pin()->HasActiveSearch())
 	{
 		return EVisibility::Visible;
