PR #6671: Enable right-click drag scrolling in Project/Editor Settings sidebar (Contributed by [REDACTED_PERSON])

diff --git a/Engine/Source/Developer/SettingsEditor/Private/Widgets/SSettingsEditor.cpp b/Engine/Source/Developer/SettingsEditor/Private/Widgets/SSettingsEditor.cpp
--- a/Engine/Source/Developer/SettingsEditor/Private/Widgets/SSettingsEditor.cpp
+++ b/Engine/Source/Developer/SettingsEditor/Private/Widgets/SSettingsEditor.cpp
@@ -74,6 +74,8 @@ void SSettingsEditor::Construct( const FArguments& InArgs, const ISettingsEditor
 	RootObjectCustomization->Initialize();
 	SettingsView->SetRootObjectCustomizationInstance(RootObjectCustomization);
 
+	TSharedPtr<SScrollBox> ScrollBox;
+
 	ChildSlot
 	[
 		SNew(SVerticalBox)
@@ -88,7 +90,7 @@ void SSettingsEditor::Construct( const FArguments& InArgs, const ISettingsEditor
 			.Padding(0.0f, 16.0f)
 			[
 				// categories menu
-				SNew(SScrollBox)
+				SAssignNew(ScrollBox, SScrollBox)
 				+ SScrollBox::Slot()
 				[
 					SNew(SHorizontalBox)
@@ -140,6 +142,8 @@ void SSettingsEditor::Construct( const FArguments& InArgs, const ISettingsEditor
 		]
 	];
 
+	ScrollBox->SetScrollBarRightClickDragAllowed(true);
+
 	FInternationalization::Get().OnCultureChanged().AddSP(this, &SSettingsEditor::HandleCultureChanged);
 	Model->OnSelectionChanged().AddSP(this, &SSettingsEditor::HandleModelSelectionChanged);
 	SettingsContainer->OnCategoryModified().AddSP(this, &SSettingsEditor::HandleSettingsContainerCategoryModified);
