[SettingsEditor]Fixed access to an invalid Section and update single property in that case

diff --git a/Engine/Source/Developer/SettingsEditor/Private/Widgets/SSettingsEditor.cpp b/Engine/Source/Developer/SettingsEditor/Private/Widgets/SSettingsEditor.cpp
--- a/Engine/Source/Developer/SettingsEditor/Private/Widgets/SSettingsEditor.cpp
+++ b/Engine/Source/Developer/SettingsEditor/Private/Widgets/SSettingsEditor.cpp
@@ -212,7 +212,7 @@ void SSettingsEditor::NotifyPostChange( const FPropertyChangedEvent& PropertyCha
 
 				if (ObjectBeingEdited->GetClass()->HasAnyClassFlags(CLASS_DefaultConfig) && !bIsArrayOrArrayElement && !bIsSetOrSetElement && !bIsMapOrMapElement)
 				{
-					if(Section->NotifySectionOnPropertyModified())
+					if (!Section.IsValid() || Section->NotifySectionOnPropertyModified())
 					{
 						ObjectBeingEdited->UpdateSinglePropertyInConfigFile(PropertyThatChanged->GetActiveMemberNode()->GetValue(), ObjectBeingEdited->GetDefaultConfigFilename());
 					}
