IWYU - Build fix

diff --git a/Engine/Source/Developer/SettingsEditor/Private/Customizations/PlatformSettingsCustomization.h b/Engine/Source/Developer/SettingsEditor/Private/Customizations/PlatformSettingsCustomization.h
--- a/Engine/Source/Developer/SettingsEditor/Private/Customizations/PlatformSettingsCustomization.h
+++ b/Engine/Source/Developer/SettingsEditor/Private/Customizations/PlatformSettingsCustomization.h
@@ -7,6 +7,7 @@
 #include "Input/Reply.h"
 #include "IDetailCustomization.h"
 #include "PropertyHandle.h"
+#include "IPropertyTypeCustomization.h"
 
 //////////////////////////////////////////////////////////////////////////
 // FPerPlatformSettingsCustomization
