Fixing redefinition, checked in old header with the same class name.

diff --git a/Engine/Source/Developer/SettingsEditor/Private/Customizations/PlatformSettingsCustomization.h b/Engine/Source/Developer/SettingsEditor/Private/Customizations/PlatformSettingsCustomization.h
deleted file mode 100644
diff --git a/Engine/Source/Developer/SettingsEditor/Private/SettingsEditorModule.cpp b/Engine/Source/Developer/SettingsEditor/Private/SettingsEditorModule.cpp
--- a/Engine/Source/Developer/SettingsEditor/Private/SettingsEditorModule.cpp
+++ b/Engine/Source/Developer/SettingsEditor/Private/SettingsEditorModule.cpp
@@ -19,7 +19,7 @@
 #include "Engine/PlatformSettings.h"
 #include "PropertyEditorModule.h"
 #include "PropertyEditorDelegates.h"
-#include "Customizations/PlatformSettingsCustomization.h"
+#include "Customizations/PerPlatformSettingsCustomization.h"
 
 #define LOCTEXT_NAMESPACE "SSettingsEditor"
 
