[Backout] - CL17906577

Original CL Desc
-----------------------------------------------------------------
Filter out any properties not marked with the "Config" metadata from Project Settings, since they won't be serialized anyway.

diff --git a/Engine/Source/Developer/SettingsEditor/Private/Widgets/SSettingsEditor.cpp b/Engine/Source/Developer/SettingsEditor/Private/Widgets/SSettingsEditor.cpp
--- a/Engine/Source/Developer/SettingsEditor/Private/Widgets/SSettingsEditor.cpp
+++ b/Engine/Source/Developer/SettingsEditor/Private/Widgets/SSettingsEditor.cpp
@@ -72,7 +72,6 @@ void SSettingsEditor::Construct( const FArguments& InArgs, const ISettingsEditor
 	SettingsView = FModuleManager::GetModuleChecked<FPropertyEditorModule>("PropertyEditor").CreateDetailView(DetailsViewArgs);
 	SettingsView->SetVisibility(TAttribute<EVisibility>::CreateSP(this, &SSettingsEditor::HandleSettingsViewVisibility));
 	SettingsView->SetIsPropertyEditingEnabledDelegate(FIsPropertyEditingEnabled::CreateSP(this, &SSettingsEditor::HandleSettingsViewEnabled));
-	SettingsView->SetIsPropertyVisibleDelegate(FIsPropertyVisible::CreateSP(this, &SSettingsEditor::IsPropertyVisible));
 
 	TSharedPtr<FSettingsDetailRootObjectCustomization> RootObjectCustomization = MakeShareable(new FSettingsDetailRootObjectCustomization(Model, SettingsView.ToSharedRef()));
 	RootObjectCustomization->Initialize();
@@ -557,6 +556,7 @@ bool SSettingsEditor::HandleSettingsViewEnabled() const
 	return (SelectedSection.IsValid() && SelectedSection->CanEdit()) || bShowingAllSettings;
 }
 
+
 EVisibility SSettingsEditor::HandleSettingsViewVisibility() const
 {
 	ISettingsSectionPtr SelectedSection = Model->GetSelectedSection();
@@ -569,22 +569,5 @@ EVisibility SSettingsEditor::HandleSettingsViewVisibility() const
 	return EVisibility::Hidden;
 }
 
-bool SSettingsEditor::IsPropertyVisible(const FPropertyAndParent& PropertyAndParent) const
-{
-	if (PropertyAndParent.Property.HasAllPropertyFlags(CPF_Config))
-	{
-		return true;
-	}
-
-	for (const FProperty* Property : PropertyAndParent.ParentProperties)
-	{
-		if (Property->HasAllPropertyFlags(CPF_Config))
-		{
-			return true;
-		}
-	}
-
-	return false;
-}
 
 #undef LOCTEXT_NAMESPACE
diff --git a/Engine/Source/Developer/SettingsEditor/Private/Widgets/SSettingsEditor.h b/Engine/Source/Developer/SettingsEditor/Private/Widgets/SSettingsEditor.h
--- a/Engine/Source/Developer/SettingsEditor/Private/Widgets/SSettingsEditor.h
+++ b/Engine/Source/Developer/SettingsEditor/Private/Widgets/SSettingsEditor.h
@@ -2,17 +2,17 @@
 
 #pragma once
 
-#include "ISettingsCategory.h"
-#include "ISettingsContainer.h"
-#include "ISettingsEditorModel.h"
-#include "ISettingsSection.h"
-#include "Layout/Visibility.h"
+#include "CoreMinimal.h"
 #include "Misc/NotifyHook.h"
-#include "PropertyEditorDelegates.h"
+#include "Layout/Visibility.h"
 #include "Widgets/DeclarativeSyntaxSupport.h"
+#include "Widgets/SWidget.h"
 #include "Widgets/SCompoundWidget.h"
 #include "Widgets/SOverlay.h"
-#include "Widgets/SWidget.h"
+#include "ISettingsCategory.h"
+#include "ISettingsContainer.h"
+#include "ISettingsEditorModel.h"
+#include "ISettingsSection.h"
 
 class IDetailsView;
 class SVerticalBox;
@@ -115,9 +115,6 @@ private:
 	/** Callback for determining the visibility of the settings view. */
 	EVisibility HandleSettingsViewVisibility() const;
 
-	/** Callback for determining the visibility of each individual property. */
-	bool IsPropertyVisible(const FPropertyAndParent& PropertyAndParent) const;
-
 	/** Callback when the timer has been ticked to refresh the categories latently. */
 	EActiveTimerReturnType UpdateCategoriesCallback(double InCurrentTime, float InDeltaTime);
 private:
