Settings Edtor: Don't list empty setting categories so they don't take space

diff --git a/Engine/Source/Developer/SettingsEditor/Private/Widgets/SSettingsEditor.cpp b/Engine/Source/Developer/SettingsEditor/Private/Widgets/SSettingsEditor.cpp
--- a/Engine/Source/Developer/SettingsEditor/Private/Widgets/SSettingsEditor.cpp
+++ b/Engine/Source/Developer/SettingsEditor/Private/Widgets/SSettingsEditor.cpp
@@ -410,12 +410,16 @@ void SSettingsEditor::ReloadCategories()
 
 		CategorySettingsSections.Sort(FSectionSortPredicate());
 
-		CategoriesBox->AddSlot()
-			.AutoHeight()
-			.Padding(0.0f, 0.0f, 0.0f, 16.0f)
-			[
-				MakeCategoryWidget(Category.ToSharedRef(), CategorySettingsSections)
-			];
+		TSharedRef<SWidget> CategoryWidget = MakeCategoryWidget(Category.ToSharedRef(), CategorySettingsSections);
+		if (CategoryWidget != SNullWidget::NullWidget)
+		{
+			CategoriesBox->AddSlot()
+				.AutoHeight()
+				.Padding(0.0f, 0.0f, 0.0f, 16.0f)
+				[
+					MoveTemp(CategoryWidget)
+				];
+		}
 
 		for(TSharedPtr<ISettingsSection>& Section : CategorySettingsSections)
 		{
