Add CVars to hide CB and EditorPrefs buttons

diff --git a/Engine/Source/Developer/SettingsEditor/Private/Widgets/SSettingsSectionHeader.cpp b/Engine/Source/Developer/SettingsEditor/Private/Widgets/SSettingsSectionHeader.cpp
--- a/Engine/Source/Developer/SettingsEditor/Private/Widgets/SSettingsSectionHeader.cpp
+++ b/Engine/Source/Developer/SettingsEditor/Private/Widgets/SSettingsSectionHeader.cpp
@@ -17,9 +17,15 @@
 #include "IDetailsView.h"
 #include "Widgets/Views/SExpanderArrow.h"
 #include "Widgets/Views/STableRow.h"
+#include "HAL/IConsoleManager.h"
 
 #define LOCTEXT_NAMESPACE "SSettingsEditor"
 
+// Workaround to hide Set As Default button until there's a better way to determine if config file is cooked
+// Currently the only way to know is to check if MakeDefaultConfigFileWritable returns false, but that can't happen before the button is clicked
+bool bHideSetAsDefaultButton = false;
+FAutoConsoleVariableRef CVarHideSetAsDefaultButton(TEXT("SettingsEditor.HideSetAsDefaultButton"), bHideSetAsDefaultButton, TEXT("Hide the Settings Editor button to save to default config."));
+
 void SSettingsSectionHeader::Construct(const FArguments& InArgs, const UObject* InSettingsObject, ISettingsEditorModelPtr InModel, TSharedPtr<IDetailsView> InDetailsView, const TSharedPtr<ITableRow>& InTableRow)
 {
 	Model = InModel;
@@ -320,7 +326,7 @@ bool SSettingsSectionHeader::HandleResetToDefaultsButtonEnabled() const
 
 EVisibility SSettingsSectionHeader::HandleSetAsDefaultButtonVisibility() const
 {
-	return (SettingsSection.IsValid() && SettingsSection->CanSaveDefaults()) ? EVisibility::Visible : EVisibility::Collapsed;
+	return (!bHideSetAsDefaultButton && (SettingsSection.IsValid() && SettingsSection->CanSaveDefaults())) ? EVisibility::Visible : EVisibility::Collapsed;
 }
 
 FReply SSettingsSectionHeader::HandleSetAsDefaultButtonClicked()
diff --git a/Engine/Source/Editor/ContentBrowser/Private/SContentBrowser.cpp b/Engine/Source/Editor/ContentBrowser/Private/SContentBrowser.cpp
--- a/Engine/Source/Editor/ContentBrowser/Private/SContentBrowser.cpp
+++ b/Engine/Source/Editor/ContentBrowser/Private/SContentBrowser.cpp
@@ -76,11 +76,16 @@
 #include "Widgets/Input/SExpandableButton.h"
 #include "SSearchToggleButton.h"
 #include "SPositiveActionButton.h"
+#include "HAL/IConsoleManager.h"
 
 #define LOCTEXT_NAMESPACE "ContentBrowser"
 
 const FString SContentBrowser::SettingsIniSection = TEXT("ContentBrowser");
 
+// Workaround to hide Save As Collection button until collections support the AliasDataSource
+bool bHideSaveSearchButton = false;
+FAutoConsoleVariableRef CVarHideSaveSearchButton(TEXT("ContentBrowser.HideSaveSearchButton"), bHideSaveSearchButton, TEXT("Hide the Content Browser button to save search as a dynamic collection."));
+
 namespace ContentBrowserSourcesWidgetSwitcherIndex
 {
 	static const int32 PathView = 0;
@@ -767,7 +772,7 @@ TSharedRef<SWidget> SContentBrowser::CreateAssetView(const FContentBrowserConfig
 							.IsEnabled(this, &SContentBrowser::IsSaveSearchButtonEnabled)
 							.OnClicked(this, &SContentBrowser::OnSaveSearchButtonClicked)
 							.ContentPadding(FMargin(1, 1))
-							.Visibility((Config != nullptr ? Config->bCanShowAssetSearch : true) ? EVisibility::Visible : EVisibility::Collapsed)
+							.Visibility((!bHideSaveSearchButton && (Config != nullptr ? Config->bCanShowAssetSearch : true)) ? EVisibility::Visible : EVisibility::Collapsed)
 							[
 								SNew(STextBlock)
 								.TextStyle(FAppStyle::Get(), "GenericFilters.TextStyle")
