Speculative fix settings editor crash on shutdown

diff --git a/Engine/Source/Developer/SettingsEditor/Private/Widgets/SSettingsEditor.cpp b/Engine/Source/Developer/SettingsEditor/Private/Widgets/SSettingsEditor.cpp
--- a/Engine/Source/Developer/SettingsEditor/Private/Widgets/SSettingsEditor.cpp
+++ b/Engine/Source/Developer/SettingsEditor/Private/Widgets/SSettingsEditor.cpp
@@ -450,6 +450,12 @@ void SSettingsEditor::HandleCultureChanged()
 
 void SSettingsEditor::HandleModelSelectionChanged()
 {
+	// This callback can trigger on unregister during shutdown, simply return in this case
+	if (!FSlateApplication::IsInitialized())
+	{
+		return;
+	}
+
 	ISettingsSectionPtr SelectedSection = Model->GetSelectedSection();
 
 	if (SelectedSection.IsValid())
