Telemetry: Fix to Editor.Usage.PreferencesChanged event to pass the unlocalized Section name to get the same value no matter what the current language is.

#rnx

diff --git a/Engine/Source/Developer/SettingsEditor/Private/Widgets/SSettingsEditor.cpp b/Engine/Source/Developer/SettingsEditor/Private/Widgets/SSettingsEditor.cpp
--- a/Engine/Source/Developer/SettingsEditor/Private/Widgets/SSettingsEditor.cpp
+++ b/Engine/Source/Developer/SettingsEditor/Private/Widgets/SSettingsEditor.cpp
@@ -338,7 +338,7 @@ void SSettingsEditor::RecordPreferenceChangedAnalytics( ISettingsSectionPtr Sele
 	if(FEngineAnalytics::IsAvailable() && ChangedProperty != nullptr && ChangedProperty->GetOwnerClass() != nullptr)
 	{
 		TArray<FAnalyticsEventAttribute> EventAttributes;
-		EventAttributes.Add(FAnalyticsEventAttribute(TEXT("PropertySection"), SelectedSection->GetDisplayName().ToString()));
+		EventAttributes.Add(FAnalyticsEventAttribute(TEXT("PropertySection"), SelectedSection->GetName().ToString()));
 		EventAttributes.Add(FAnalyticsEventAttribute(TEXT("PropertyClass"), ChangedProperty->GetOwnerClass()->GetName()));
 		EventAttributes.Add(FAnalyticsEventAttribute(TEXT("PropertyName"), ChangedProperty->GetName()));
 
