Add delegate to control which settings get auto-discovered by the settings module

diff --git a/Engine/Source/Developer/SettingsEditor/Private/SettingsEditorModule.cpp b/Engine/Source/Developer/SettingsEditor/Private/SettingsEditorModule.cpp
--- a/Engine/Source/Developer/SettingsEditor/Private/SettingsEditorModule.cpp
+++ b/Engine/Source/Developer/SettingsEditor/Private/SettingsEditorModule.cpp
@@ -186,6 +186,11 @@ public:
 		ApplicationRestartRequiredNotification.SetOnRestartApplicationCallback(InRestartApplicationDelegate);
 	}
 
+	virtual void SetShouldRegisterSettingCallback(FShouldRegisterSettingsDelegate InShouldRegisterSettingDelegate) override
+	{
+		ShouldRegisterSettingsDelegate = MoveTemp(InShouldRegisterSettingDelegate);
+	}
+
 private:
 
 	void ModulesChangesCallback(FName ModuleName, EModuleChangeReason ReasonForChange)
@@ -246,6 +251,12 @@ private:
 						continue;
 					}
 
+					// Let an external module decide if the setting should be registered.
+					if (ShouldRegisterSettingsDelegate.IsBound() && !ShouldRegisterSettingsDelegate.Execute(Settings))
+					{
+						continue;
+					}
+
 					RegisterDeveloperSettings(SettingsModule, Settings);
 				}
 			}
@@ -305,6 +316,9 @@ private:
 
 	/** Flag if the settings are stale currently and need to be refreshed. */
 	bool bAreSettingsStale;
+
+	/** Delegate called to decide if a settings object should be registered with the module. */
+	FShouldRegisterSettingsDelegate ShouldRegisterSettingsDelegate;
 };
 
 
diff --git a/Engine/Source/Developer/SettingsEditor/Public/ISettingsEditorModule.h b/Engine/Source/Developer/SettingsEditor/Public/ISettingsEditorModule.h
--- a/Engine/Source/Developer/SettingsEditor/Public/ISettingsEditorModule.h
+++ b/Engine/Source/Developer/SettingsEditor/Public/ISettingsEditorModule.h
@@ -3,11 +3,15 @@
 #pragma once
 
 #include "CoreMinimal.h"
-#include "Widgets/SWidget.h"
+#include "Delegates/DelegateCombinations.h"
 #include "Modules/ModuleInterface.h"
+#include "Widgets/SWidget.h"
 
 class ISettingsContainer;
 class ISettingsEditorModel;
+class UDeveloperSettings;
+
+DECLARE_DELEGATE_RetVal_OneParam(bool, FShouldRegisterSettingsDelegate , UDeveloperSettings* /*Settings*/);
 
 /**
  * Interface for settings editor modules.
@@ -47,6 +51,13 @@ public:
 	 */
 	virtual void SetRestartApplicationCallback( FSimpleDelegate InRestartApplicationDelegate ) = 0;
 
+	/**
+	 * Set the delegate that should be called when a setting editor checks whether a settings object should be registered.
+	 * 
+	 * @param InShouldRegisterSettingDelegate The new delegate to call.
+	 */
+	virtual void SetShouldRegisterSettingCallback(FShouldRegisterSettingsDelegate InShouldRegisterSettingDelegate) = 0;
+
 public:
 
 	/** Virtual destructor. */
