ISettingsEditorModule : Expose UpdateSettings as part of the public module interface.

This allows editors in other modules to force the registration of auto-discovered UDeveloperSettings so that they can work with the complete set of settings.  This is not called internally until the first time a settings editor ( e.g. Editor Preferences ) is shown.

diff --git a/Engine/Source/Developer/SettingsEditor/Private/SettingsEditorModule.cpp b/Engine/Source/Developer/SettingsEditor/Private/SettingsEditorModule.cpp
--- a/Engine/Source/Developer/SettingsEditor/Private/SettingsEditorModule.cpp
+++ b/Engine/Source/Developer/SettingsEditor/Private/SettingsEditorModule.cpp
@@ -200,7 +200,7 @@ private:
 		UpdateSettings();
 	}
 
-	void UpdateSettings(bool bForce = false)
+	virtual void UpdateSettings(bool bForce = false) override
 	{
 		if ( ( AnyActiveSettingsEditor() || bForce ) && bAreSettingsStale )
 		{
diff --git a/Engine/Source/Developer/SettingsEditor/Public/ISettingsEditorModule.h b/Engine/Source/Developer/SettingsEditor/Public/ISettingsEditorModule.h
--- a/Engine/Source/Developer/SettingsEditor/Public/ISettingsEditorModule.h
+++ b/Engine/Source/Developer/SettingsEditor/Public/ISettingsEditorModule.h
@@ -58,6 +58,13 @@ public:
 	 */
 	virtual void SetShouldRegisterSettingCallback(FShouldRegisterSettingsDelegate InShouldRegisterSettingDelegate) = 0;
 
+	/**
+	 * Registers any pending auto-discovered settings.
+	 *
+	 * @param bForce Forces the registration of settings even if there is no active settings editor.
+	 */
+	virtual void UpdateSettings(bool bForce = false) = 0;
+
 public:
 
 	/** Virtual destructor. */
