Tweaks to the Settings UI (Editor Preferences & Project Settings)
- Made it so that hidden buttons don't take space
- Added middle ellipsis on text
- Reworked the layout a bit so that the category title isn't affected by the button row. The title is on its own row and the description and buttons are on a second row
- Made the revision control Locked / Unlocked icons respect the foreground color
- Aligned the right side of the header with the rest of the table UI (removed the 16px right padding)
- Buttons and checkout notice will now get collapsed when the table row isn't expanded

diff --git a/Engine/Source/Developer/SettingsEditor/Private/Widgets/SSettingsSectionHeader.cpp b/Engine/Source/Developer/SettingsEditor/Private/Widgets/SSettingsSectionHeader.cpp
--- a/Engine/Source/Developer/SettingsEditor/Private/Widgets/SSettingsSectionHeader.cpp
+++ b/Engine/Source/Developer/SettingsEditor/Private/Widgets/SSettingsSectionHeader.cpp
@@ -45,97 +45,108 @@ void SSettingsSectionHeader::Construct(const FArguments& InArgs, const UObject*
 
 
 	ChildSlot
-	.Padding(FMargin(0.0f, 8.0f, 16.0f, 5.0f))
+	.Padding(FMargin(0.0f, 8.0f, 0.0f, 5.0f))
 	[
 		SNew(SVerticalBox)
 		+ SVerticalBox::Slot()
 		.AutoHeight()
 		[
-			SNew(SHorizontalBox)
-			+ SHorizontalBox::Slot()
-			.FillWidth(1.0f)
+			SNew(SVerticalBox)
+			+SVerticalBox::Slot()
+			.FillHeight(1.0f)
 			[
-				SNew(SVerticalBox)
-				+ SVerticalBox::Slot()
-				.AutoHeight()
+				SNew(SHorizontalBox)
+				+ SHorizontalBox::Slot()
+				.VAlign(VAlign_Center)
+				.Padding(2.0f, 2.0f, 2.0f, 2.0f)
+				.AutoWidth()
 				[
-					SNew(SHorizontalBox)
-					+ SHorizontalBox::Slot()
-					.VAlign(VAlign_Center)
-					.Padding(2.0f, 2.0f, 2.0f, 2.0f)
-					.AutoWidth()
-					[
-						SNew(SExpanderArrow, InTableRow)
-					]
-					+ SHorizontalBox::Slot()// category title
-					[
-						SNew(STextBlock)
-						.Font(FAppStyle::GetFontStyle("SettingsEditor.CatgoryAndSectionFont"))
-						.Text(GetSettingsBoxTitleText())
-					]
+					SNew(SExpanderArrow, InTableRow)
 				]
-				+ SVerticalBox::Slot()
-				.AutoHeight()
-				.Padding(16.0f, 8.0f, 0.0f, 0.0f)
+				+ SHorizontalBox::Slot()// category title
 				[
-					// category description
 					SNew(STextBlock)
-					.ColorAndOpacity(FSlateColor::UseSubduedForeground())
-					.Text(GetSettingsBoxDescriptionText())
-					.Visibility(this, &SSettingsSectionHeader::GetCategoryDescriptionVisibility)
+					.Font(FAppStyle::GetFontStyle("SettingsEditor.CatgoryAndSectionFont"))
+					.Text(GetSettingsBoxTitleText())
+					.OverflowPolicy(ETextOverflowPolicy::MiddleEllipsis)
 				]
 			]
-			+ SHorizontalBox::Slot()
-			.AutoWidth()
-			.HAlign(HAlign_Right)
-			.VAlign(VAlign_Bottom)
-			.Padding(16.0f, 0.0f, 0.0f, 0.0f)
+			+SVerticalBox::Slot()
+			.Padding(FMargin(0.f, 8.f, 0.f, 0.f))
+			.AutoHeight()
+			.MinHeight(TAttribute<float>(this, &SSettingsSectionHeader::GetDescriptionRowMinHeight))
 			[
 				SNew(SHorizontalBox)
-				.Visibility(this, &SSettingsSectionHeader::GetButtonRowVisibility)
-				+SHorizontalBox::Slot()
-				[
-					// set as default button
-					SNew(SButton)
-					.Visibility(this, &SSettingsSectionHeader::HandleSetAsDefaultButtonVisibility)
-					.IsEnabled(this, &SSettingsSectionHeader::HandleSetAsDefaultButtonEnabled)
-					.OnClicked(this, &SSettingsSectionHeader::HandleSetAsDefaultButtonClicked)
-					.Text(LOCTEXT("SaveDefaultsButtonText", "Set as Default"))
-					.ToolTipText(LOCTEXT("SaveDefaultsButtonTooltip", "Save the values below as the new default settings"))
-				]
 				+ SHorizontalBox::Slot()
-				.AutoWidth()
-				.Padding(8.0f, 0.0f, 0.0f, 0.0f)
+				.FillWidth(1.0f)
+				.VAlign(VAlign_Fill)
+				.Padding(20.0f, 0.0f, 0.0f, 0.0f)
 				[
-					// export button
-					SNew(SButton)
-					.IsEnabled(this, &SSettingsSectionHeader::HandleExportButtonEnabled)
-					.OnClicked(this, &SSettingsSectionHeader::HandleExportButtonClicked)
-					.Text(LOCTEXT("ExportButtonText", "Export..."))
-					.ToolTipText(LOCTEXT("ExportButtonTooltip", "Export these settings to a file on your computer"))
-				]
-				+ SHorizontalBox::Slot()
-				.AutoWidth()
-				.Padding(8.0f, 0.0f, 0.0f, 0.0f)
-				[
-					// import button
-					SNew(SButton)
-					.IsEnabled(this, &SSettingsSectionHeader::HandleImportButtonEnabled)
-					.OnClicked(this, &SSettingsSectionHeader::HandleImportButtonClicked)
-					.Text(LOCTEXT("ImportButtonText", "Import..."))
-					.ToolTipText(LOCTEXT("ImportButtonTooltip", "Import these settings from a file on your computer"))
+					SNew(SBorder)
+					.Padding(2.f, 0.f, 0.f, 0.f)
+					.VAlign(VAlign_Center)
+					.BorderImage(FAppStyle::GetBrush("NoBorder"))
+					.Visibility(this, &SSettingsSectionHeader::GetContentVisibility)
+					[
+						// category description
+						SNew(STextBlock)
+						.ColorAndOpacity(FSlateColor::UseSubduedForeground())
+						.Text(GetSettingsBoxDescriptionText())
+						.OverflowPolicy(ETextOverflowPolicy::Ellipsis)
+					]
 				]
 				+ SHorizontalBox::Slot()
 				.AutoWidth()
-				.Padding(8.0f, 0.0f, 0.0f, 0.0f)
+				.HAlign(HAlign_Right)
+				.VAlign(VAlign_Center)
+				.Padding(16.0f, 0.0f, 0.0f, 0.0f)
 				[
-					// reset defaults button
-					SNew(SButton)
-					.Visibility(this, &SSettingsSectionHeader::HandleSetAsDefaultButtonVisibility)
-					.IsEnabled(this, &SSettingsSectionHeader::HandleResetToDefaultsButtonEnabled)
-					.OnClicked(this, &SSettingsSectionHeader::HandleResetDefaultsButtonClicked)
-					.Text(LOCTEXT("ResetDefaultsButtonText", "Reset to Defaults"))
-					.ToolTipText(LOCTEXT("ResetDefaultsButtonTooltip", "Reset the settings below to their default values"))
+					SNew(SHorizontalBox)
+					.Visibility(this, &SSettingsSectionHeader::GetButtonRowVisibility)
+					+SHorizontalBox::Slot()
+					[
+						// set as default button
+						SNew(SButton)
+						.Visibility(this, &SSettingsSectionHeader::HandleSetAsDefaultButtonVisibility)
+						.IsEnabled(this, &SSettingsSectionHeader::HandleSetAsDefaultButtonEnabled)
+						.OnClicked(this, &SSettingsSectionHeader::HandleSetAsDefaultButtonClicked)
+						.Text(LOCTEXT("SaveDefaultsButtonText", "Set as Default"))
+						.ToolTipText(LOCTEXT("SaveDefaultsButtonTooltip", "Save the values below as the new default settings"))
+					]
+					+SHorizontalBox::Slot()
+					.AutoWidth()
+					.Padding(8.0f, 0.0f, 0.0f, 0.0f)
+					[
+						// export button
+						SNew(SButton)
+						.IsEnabled(this, &SSettingsSectionHeader::HandleExportButtonEnabled)
+						.OnClicked(this, &SSettingsSectionHeader::HandleExportButtonClicked)
+						.Text(LOCTEXT("ExportButtonText", "Export..."))
+						.ToolTipText(LOCTEXT("ExportButtonTooltip", "Export these settings to a file on your computer"))
+					]
+					+SHorizontalBox::Slot()
+					.AutoWidth()
+					.Padding(8.0f, 0.0f, 0.0f, 0.0f)
+					[
+						// import button
+						SNew(SButton)
+						.IsEnabled(this, &SSettingsSectionHeader::HandleImportButtonEnabled)
+						.OnClicked(this, &SSettingsSectionHeader::HandleImportButtonClicked)
+						.Text(LOCTEXT("ImportButtonText", "Import..."))
+						.ToolTipText(LOCTEXT("ImportButtonTooltip", "Import these settings from a file on your computer"))
+					]
+					+SHorizontalBox::Slot()
+					.AutoWidth()
+					.Padding(8.0f, 0.0f, 0.0f, 0.0f)
+					[
+						// reset defaults button
+						SNew(SButton)
+						.Visibility(this, &SSettingsSectionHeader::HandleSetAsDefaultButtonVisibility)
+						.IsEnabled(this, &SSettingsSectionHeader::HandleResetToDefaultsButtonEnabled)
+						.OnClicked(this, &SSettingsSectionHeader::HandleResetDefaultsButtonClicked)
+						.Text(LOCTEXT("ResetDefaultsButtonText", "Reset to Defaults"))
+						.ToolTipText(LOCTEXT("ResetDefaultsButtonTooltip", "Reset the settings below to their default values"))
+					]
 				]
 			]
 		]
@@ -172,10 +183,19 @@ FText SSettingsSectionHeader::GetSettingsBoxDescriptionText() const
 
 EVisibility SSettingsSectionHeader::GetButtonRowVisibility() const
 {
-	return DetailsView.Pin()->HasActiveSearch() ? EVisibility::Hidden : EVisibility::Visible;
+	const EVisibility DefaultVisibility = GetContentVisibility();
+
+	if (DefaultVisibility != EVisibility::Visible)
+	{
+		return DefaultVisibility;
+	}
+	else
+	{
+		return DetailsView.Pin()->HasActiveSearch() ? EVisibility::Collapsed : EVisibility::Visible;
+	}
 }
 
-EVisibility SSettingsSectionHeader::GetCategoryDescriptionVisibility() const
+EVisibility SSettingsSectionHeader::GetContentVisibility() const
 {
 	return TableRow.IsValid() && TableRow.Pin()->IsItemExpanded() ? EVisibility::Visible : EVisibility::Collapsed;
 }
@@ -212,6 +232,19 @@ FReply SSettingsSectionHeader::HandleExportButtonClicked()
 }
 
 
+float SSettingsSectionHeader::GetDescriptionRowMinHeight() const
+{
+	if (GetContentVisibility() == EVisibility::Visible)
+	{
+		return 22.f; // so that we keep a constant height even if the buttons are collapsed
+	}
+	else
+	{
+		return 0.f;
+	}
+}
+
+
 bool SSettingsSectionHeader::HandleExportButtonEnabled() const
 {
 	if(SettingsSection.IsValid())
@@ -512,6 +545,13 @@ void SSettingsSectionHeader::HandleCheckoutNoticeFileProbablyModifiedExternally(
 /** Callback for determining the visibility of the 'Locked' notice. */
 EVisibility SSettingsSectionHeader::GetCheckoutNoticeVisibility() const
 {
+	const EVisibility DefaultVisibility = GetContentVisibility();
+
+	if (DefaultVisibility != EVisibility::Visible)
+	{
+		return DefaultVisibility;
+	}
+
 	// Only DefaultConfig are under source control, so the checkout notice should not be visible for the other cases
 	if(SettingsObject.IsValid() && SettingsObject->GetClass()->HasAnyClassFlags(CLASS_DefaultConfig) && !DetailsView.Pin()->HasActiveSearch())
 	{
diff --git a/Engine/Source/Developer/SettingsEditor/Private/Widgets/SSettingsSectionHeader.h b/Engine/Source/Developer/SettingsEditor/Private/Widgets/SSettingsSectionHeader.h
--- a/Engine/Source/Developer/SettingsEditor/Private/Widgets/SSettingsSectionHeader.h
+++ b/Engine/Source/Developer/SettingsEditor/Private/Widgets/SSettingsSectionHeader.h
@@ -95,7 +95,9 @@ private:
 
 	EVisibility GetButtonRowVisibility() const;
 
-	EVisibility GetCategoryDescriptionVisibility() const;
+	EVisibility GetContentVisibility() const;
+
+	float GetDescriptionRowMinHeight() const;
 
 	void OnSettingsSelectionChanged();
 private:
diff --git a/Engine/Source/Developer/SharedSettingsWidgets/Private/SSettingsEditorCheckoutNotice.cpp b/Engine/Source/Developer/SharedSettingsWidgets/Private/SSettingsEditorCheckoutNotice.cpp
--- a/Engine/Source/Developer/SharedSettingsWidgets/Private/SSettingsEditorCheckoutNotice.cpp
+++ b/Engine/Source/Developer/SharedSettingsWidgets/Private/SSettingsEditorCheckoutNotice.cpp
@@ -191,6 +191,7 @@ void SSettingsEditorCheckoutNotice::Construct( const FArguments& InArgs )
 			[
 				SNew(SImage)
 				.Image(FAppStyle::Get().GetBrush("Icons.Lock"))
+				.ColorAndOpacity(FSlateColor::UseForeground())
 			]
 
 			// Locked notice
@@ -202,7 +203,6 @@ void SSettingsEditorCheckoutNotice::Construct( const FArguments& InArgs )
 				SNew(STextBlock)
 				.Text(this, &SSettingsEditorCheckoutNotice::HandleLockedStatusText)
 				.ColorAndOpacity(this, &SSettingsEditorCheckoutNotice::GetLockedStatusTextColor)
-								
 			]
 
 			// Check out button
@@ -240,6 +240,7 @@ void SSettingsEditorCheckoutNotice::Construct( const FArguments& InArgs )
 			[
 				SNew(SImage)
 				.Image(FAppStyle::Get().GetBrush("Icons.Unlock"))
+				.ColorAndOpacity(FSlateColor::UseForeground())
 			]
 
 			// Unlocked notice
@@ -308,7 +309,7 @@ EVisibility SSettingsEditorCheckoutNotice::HandleCheckOutButtonVisibility() cons
 		return EVisibility::Visible;
 	}
 
-	return EVisibility::Hidden;
+	return EVisibility::Collapsed;
 }
 
 FText SSettingsEditorCheckoutNotice::HandleLockedStatusText() const
