SSettingsEditor : Workaround rendering issue when the UpdateCategoriesCallback is called during startup

diff --git a/Engine/Source/Developer/SettingsEditor/Private/Widgets/SSettingsEditor.cpp b/Engine/Source/Developer/SettingsEditor/Private/Widgets/SSettingsEditor.cpp
--- a/Engine/Source/Developer/SettingsEditor/Private/Widgets/SSettingsEditor.cpp
+++ b/Engine/Source/Developer/SettingsEditor/Private/Widgets/SSettingsEditor.cpp
@@ -23,6 +23,7 @@
 #include "SSettingsEditorCheckoutNotice.h"
 #include "HAL/PlatformFileManager.h"
 #include "HAL/PlatformFile.h"
+#include "Containers/Ticker.h"
 
 #define LOCTEXT_NAMESPACE "SSettingsEditor"
 
@@ -546,17 +547,17 @@ void SSettingsEditor::HandleSettingsContainerCategoryModified( const FName& Cate
 	if ( !bIsActiveTimerRegistered )
 	{
 		bIsActiveTimerRegistered = true;
-		RegisterActiveTimer(0.f, FWidgetActiveTimerDelegate::CreateSP(this, &SSettingsEditor::UpdateCategoriesCallback));
+		FTSTicker::GetCoreTicker().AddTicker(FTickerDelegate::CreateSP(this, &SSettingsEditor::UpdateCategoriesCallback));
 	}
 }
 
-EActiveTimerReturnType SSettingsEditor::UpdateCategoriesCallback(double InCurrentTime, float InDeltaTime)
+bool SSettingsEditor::UpdateCategoriesCallback(float InDeltaTime)
 {
 	bIsActiveTimerRegistered = false;
 
 	ReloadCategories();
 
-	return EActiveTimerReturnType::Stop;
+	return false;
 }
 
 bool SSettingsEditor::HandleSettingsViewEnabled() const
diff --git a/Engine/Source/Developer/SettingsEditor/Private/Widgets/SSettingsEditor.h b/Engine/Source/Developer/SettingsEditor/Private/Widgets/SSettingsEditor.h
--- a/Engine/Source/Developer/SettingsEditor/Private/Widgets/SSettingsEditor.h
+++ b/Engine/Source/Developer/SettingsEditor/Private/Widgets/SSettingsEditor.h
@@ -116,7 +116,7 @@ private:
 	EVisibility HandleSettingsViewVisibility() const;
 
 	/** Callback when the timer has been ticked to refresh the categories latently. */
-	EActiveTimerReturnType UpdateCategoriesCallback(double InCurrentTime, float InDeltaTime);
+	bool UpdateCategoriesCallback(float InDeltaTime);
 private:
 
 	/** Holds the vertical box for settings categories. */
diff --git a/Engine/Source/Developer/SettingsEditor/Private/Widgets/SSettingsSectionHeader.h b/Engine/Source/Developer/SettingsEditor/Private/Widgets/SSettingsSectionHeader.h
--- a/Engine/Source/Developer/SettingsEditor/Private/Widgets/SSettingsSectionHeader.h
+++ b/Engine/Source/Developer/SettingsEditor/Private/Widgets/SSettingsSectionHeader.h
@@ -26,8 +26,6 @@ public:
 
 private:
 
-	void CommonConstruct(const UObject* InSettingsObject, ISettingsEditorModelPtr InModel, TSharedPtr<IDetailsView> InDetailsView);
-
 	FText GetSettingsBoxTitleText() const;
 	FText GetSettingsBoxDescriptionText() const;
 
