Settings: Fix crash due to settings being refreshed during PostLoad.

This can happen when a module, in this case the landscape system, loads a module during PostLoad. Loading a module triggers a refresh inside the settings module. This fix simply defers the udpdate to the next ticket.
I've also added some defensive code to avoid defering it more than 10 times in case somebody edits the defer condition in the future and accidentally creates an infinite loop of defers.

diff --git a/Engine/Source/Developer/SettingsEditor/Private/SettingsEditorModule.cpp b/Engine/Source/Developer/SettingsEditor/Private/SettingsEditorModule.cpp
--- a/Engine/Source/Developer/SettingsEditor/Private/SettingsEditorModule.cpp
+++ b/Engine/Source/Developer/SettingsEditor/Private/SettingsEditorModule.cpp
@@ -20,6 +20,7 @@
 #include "PropertyEditorModule.h"
 #include "PropertyEditorDelegates.h"
 #include "Customizations/PerPlatformSettingsCustomization.h"
+#include "UObject/UObjectThreadContext.h"
 
 #define LOCTEXT_NAMESPACE "SSettingsEditor"
 
@@ -197,7 +198,7 @@ private:
 	{
 		ClearStaleEditorWidgets();
 		bAreSettingsStale = true;
-		UpdateSettings();
+		UpdateSettingsNowOrDeferIfUnsafe();
 	}
 
 	virtual void UpdateSettings(bool bForce = false) override
@@ -228,6 +229,38 @@ private:
 		}
 	}
 
+	void UpdateSettingsNowOrDeferIfUnsafe()
+	{
+		// Updating can cause UObject::ProcessEvent to be called, e.g. via FDetailPropertyRow::IsEditConditionMet.
+		// UObject::ProcessEvent is not allowed to be called during PostLoad.
+		// We can get here by having Editor Settings open and a UObject's PostLoad loading a module during PostLoad, which triggers
+		// ModulesChangesCallback and in turns calls UpdateSettings.
+		const bool bIsSafeToUpdateNow = !FUObjectThreadContext::Get().IsRoutingPostLoad;
+		if (bIsSafeToUpdateNow)
+		{
+			UpdateSettings();
+		}
+		else
+		{
+			DeferUpdateSettingsToNextTick();
+		}
+	}
+	
+	void DeferUpdateSettingsToNextTick()
+	{
+		if (bHasDeferredUpdate)
+		{
+			return;
+		}
+		
+		bHasDeferredUpdate = true;
+		ExecuteOnGameThread(TEXT("UpdateSettings"), [this]
+		{
+			bHasDeferredUpdate = false;
+			UpdateSettingsNowOrDeferIfUnsafe();
+		});
+	}
+	
 	bool AnyActiveSettingsEditor()
 	{
 		return EditorWidgets.Num() > 0;
@@ -317,6 +350,9 @@ private:
 	/** Flag if the settings are stale currently and need to be refreshed. */
 	bool bAreSettingsStale;
 
+	/** Flag to see whether the settings should be updated next tick. */
+	bool bHasDeferredUpdate = false;
+
 	/** Delegate called to decide if a settings object should be registered with the module. */
 	FShouldRegisterSettingsDelegate ShouldRegisterSettingsDelegate;
 };
diff --git a/Engine/Source/Editor/PropertyEditor/Private/EditConditionParser.cpp b/Engine/Source/Editor/PropertyEditor/Private/EditConditionParser.cpp
--- a/Engine/Source/Editor/PropertyEditor/Private/EditConditionParser.cpp
+++ b/Engine/Source/Editor/PropertyEditor/Private/EditConditionParser.cpp
@@ -15,6 +15,7 @@
 #include "Misc/Optional.h"
 #include "Trace/Detail/Channel.h"
 #include "UObject/NameTypes.h"
+#include "UObject/UObjectThreadContext.h"
 
 class UObject;
 
@@ -834,6 +835,14 @@ TValueOrError<bool, FText> FEditConditionParser::Evaluate(const FEditConditionEx
 {
 	using namespace EditConditionParserTokens;
 	
+	// Edit conditions may invoke UFUNCTION()s by going through UObject::ProcessEvent, which checkf()s this condition.
+	const bool bIsSafeToCall = !FUObjectThreadContext::Get().IsRoutingPostLoad;
+	// If you hit this ensure, then adjust your logic to not call this during PostLoad, e.g. try deferring it to end of frame / next tick.
+	if (!ensure(bIsSafeToCall))
+	{
+		return MakeError(LOCTEXT("Evaluate.InternalError", "Internal Error: Invalid call timing."));
+	}
+	
 	FExpressionResult Result = ExpressionParser::Evaluate(Expression.Tokens, OperatorJumpTable, &Context);
 	if (Result.HasValue())
 	{
