Editor/Project Settings:
- Added top-level Export and Import buttons to apply to all settings.
- Now shows All Settings by default, instead of selecting the first category.

[RN]

diff --git a/Engine/Source/Developer/SettingsEditor/Private/Widgets/SSettingsEditor.cpp b/Engine/Source/Developer/SettingsEditor/Private/Widgets/SSettingsEditor.cpp
--- a/Engine/Source/Developer/SettingsEditor/Private/Widgets/SSettingsEditor.cpp
+++ b/Engine/Source/Developer/SettingsEditor/Private/Widgets/SSettingsEditor.cpp
@@ -13,6 +13,7 @@
 #include "Widgets/Layout/SScrollBox.h"
 #include "Styling/AppStyle.h"
 #include "AnalyticsEventAttribute.h"
+#include "DesktopPlatformModule.h"
 #include "EngineAnalytics.h"
 #include "Interfaces/IAnalyticsProvider.h"
 #include "Widgets/Text/STextBlock.h"
@@ -24,6 +25,10 @@
 #include "HAL/PlatformFileManager.h"
 #include "HAL/PlatformFile.h"
 #include "Containers/Ticker.h"
+#include "Framework/Notifications/NotificationManager.h"
+#include "HAL/FileManager.h"
+#include "Misc/ConfigCacheIni.h"
+#include "Misc/MessageDialog.h"
 
 #define LOCTEXT_NAMESPACE "SSettingsEditor"
 
@@ -128,7 +133,34 @@ void SSettingsEditor::Construct( const FArguments& InArgs, const ISettingsEditor
 				+ SVerticalBox::Slot()
 				.AutoHeight()
 				[
-					SettingsView->GetFilterAreaWidget().ToSharedRef()
+					SNew(SHorizontalBox)
+					+ SHorizontalBox::Slot()
+					.FillWidth(1.0f)
+					[
+						SettingsView->GetFilterAreaWidget().ToSharedRef()
+					]
+					+ SHorizontalBox::Slot()
+					.AutoWidth()
+					.VAlign(VAlign_Center)
+					.Padding(8.0f, 0.0f, 0.0f, 0.0f)
+					[
+						// Export button
+						SNew(SButton)
+						.OnClicked(this, &SSettingsEditor::HandleExportButtonClicked)
+						.Text(LOCTEXT("ExportButtonText", "Export..."))
+						.ToolTipText(LOCTEXT("ExportButtonTooltip", "Export all editor settings to a file on your computer"))
+					]
+					+ SHorizontalBox::Slot()
+					.AutoWidth()
+					.VAlign(VAlign_Center)
+					.Padding(8.0f, 0.0f, 6.0f, 0.0f)
+					[
+						// Import button
+						SNew(SButton)
+						.OnClicked(this, &SSettingsEditor::HandleImportButtonClicked)
+						.Text(LOCTEXT("ImportButtonText", "Import..."))
+						.ToolTipText(LOCTEXT("ImportButtonTooltip", "Import all editor settings from a file on your computer"))
+					]
 				]
 				+ SVerticalBox::Slot()
 				.FillHeight(1.0f)
@@ -439,6 +471,7 @@ void SSettingsEditor::ReloadCategories()
 		}
 	}
 
+	Model->SelectSection(nullptr);
 	SettingsView->SetObjects(SettingsObjects);
 
 
@@ -567,6 +600,158 @@ bool SSettingsEditor::UpdateCategoriesCallback(float InDeltaTime)
 	return false;
 }
 
+void ShowNotification(const FText& Text, SNotificationItem::ECompletionState CompletionState)
+{
+	FNotificationInfo Notification(Text);
+	Notification.ExpireDuration = 3.f;
+	Notification.bUseSuccessFailIcons = false;
+
+	FSlateNotificationManager::Get().AddNotification(Notification)->SetCompletionState(CompletionState);
+}
+
+FReply SSettingsEditor::HandleExportButtonClicked()
+{
+	const FString DefaultFileName = FString::Printf(TEXT("%s Settings Backup %s.ini"), *Model->GetSettingsContainer()->GetName().ToString(), *FDateTime::Now().ToString(TEXT("%Y-%m-%d %H%M%S")));
+
+	const TSharedPtr<SWindow> ParentWindow = FSlateApplication::Get().FindWidgetWindow(AsShared());
+	const void* ParentWindowHandle = ParentWindow.IsValid() && ParentWindow->GetNativeWindow().IsValid() ? ParentWindow->GetNativeWindow()->GetOSWindowHandle() : nullptr;
+
+	TArray<FString> OutSaveFiles;
+	if (FDesktopPlatformModule::Get()->SaveFileDialog(ParentWindowHandle, 
+		LOCTEXT("ExportSettingsDialogTitle", "Export settings...").ToString(), 
+		FPaths::GetPath(GEditorPerProjectIni),
+		DefaultFileName, 
+		TEXT("Config files (*.ini)|*.ini"), 
+		EFileDialogFlags::None, 
+		OutSaveFiles))
+	{
+		const FString DestinationFilePath = OutSaveFiles[0];
+		if (GConfig->FindBranch(*DestinationFilePath, DestinationFilePath))
+		{
+			ShowNotification(LOCTEXT("ExportSettingsExistingFileFailure", "Export settings failed. Cannot use a file that is already part of the config system."), SNotificationItem::CS_Fail);
+			return FReply::Handled();
+		}
+
+		FString DestinationFolder;
+		FString DestinationExtension;
+		FString DestinationFilenameWithoutExtension;
+		FPaths::Split(DestinationFilePath, DestinationFolder, DestinationFilenameWithoutExtension, DestinationExtension);
+
+		// Exports will be done to a temp folder so we can be aware of any extra files generated via custom Export delegates within the destination folder
+		const FString TempDestinationFolder = FPaths::ProjectSavedDir() / TEXT("TempExport");
+
+		constexpr bool bDeleteDirectoryRequireExists = false;
+		constexpr bool bDeleteDirectoryRecursively = true;
+
+		// If the TempExport folder already exists, prompt the user to delete it.
+		// This should really only happen if there was a crash during export or the user manually created it.
+		if (FPaths::DirectoryExists(TempDestinationFolder))
+		{
+			if (FMessageDialog::Open(EAppMsgType::OkCancel,
+				LOCTEXT("ExportSettingsDeleteTempFolderMsg", "TempExport settings folder already exists within project's Saved folder, possibly from a previous failed export. Do you wish to delete this folder so this export may continue?"),
+				LOCTEXT("ExportSettingsDeleteTempFolderTitle", "Delete Folder?")) == EAppReturnType::Ok)
+			{
+				IFileManager::Get().DeleteDirectory(*TempDestinationFolder, bDeleteDirectoryRequireExists, bDeleteDirectoryRecursively);
+			}
+			else
+			{
+				ShowNotification(LOCTEXT("ExportSettingsExistingTempFolderFailure", "Export settings failed. TempExport folder already exists under project Saved folder."), SNotificationItem::CS_Fail);
+				return FReply::Handled();
+			}
+		}
+
+		FConfigBranch& TempBranch = GConfig->AddNewBranch(DestinationFilePath);
+
+		// A temp destination file is used to accomodate custom Export delegates that use intermediate files that are then copied into the destination file
+		const FString TempDestinationFilePath = TempDestinationFolder / DestinationFilenameWithoutExtension + TEXT(".") + DestinationExtension;
+
+		TArray<ISettingsCategoryPtr> Categories;
+		SettingsContainer->GetCategories(Categories);
+
+		for (const ISettingsCategoryPtr& Category : Categories)
+		{
+			TArray<ISettingsSectionPtr> Sections;
+			Category->GetSections(Sections);
+			Sections.Sort([](const ISettingsSectionPtr& First, const ISettingsSectionPtr& Second)
+			{
+				return First->GetDisplayName().CompareTo(Second->GetDisplayName()) < 0;
+			});
+
+			for (const ISettingsSectionPtr& Section : Sections)
+			{
+				if (!Section->Export(*TempDestinationFilePath))
+				{
+					continue;
+				}
+
+				TempBranch.InMemoryFile.Combine(TempDestinationFilePath);
+				IFileManager::Get().Delete(*TempDestinationFilePath);
+
+				TArray<FString> AdditionalFiles;
+				IFileManager::Get().FindFiles(AdditionalFiles, *TempDestinationFolder);
+
+				// Copy any additional files to the destination file's folder with a unique name based on the destination file path and the section it was in
+				for (const FString& AdditionalFile : AdditionalFiles)
+				{
+					FString AdditionalFilePath = TempDestinationFolder / AdditionalFile;
+					FString NewAdditionalFilePath = FString::Printf(TEXT("%s/%s_%s_%s"), *DestinationFolder, *DestinationFilenameWithoutExtension, *Section->GetDisplayName().ToString(), *AdditionalFile);
+					IFileManager::Get().Copy(*NewAdditionalFilePath, *AdditionalFilePath);
+					IFileManager::Get().Delete(*AdditionalFilePath);
+				}
+			}
+		}
+
+		GConfig->Flush(true, *DestinationFilePath);
+
+		IFileManager::Get().DeleteDirectory(*TempDestinationFolder, bDeleteDirectoryRequireExists, bDeleteDirectoryRecursively);
+
+		ShowNotification(LOCTEXT("ExportSettingsSuccess", "Export settings succeeded"), SNotificationItem::CS_Success);
+		return FReply::Handled();
+	}
+
+	ShowNotification(LOCTEXT("ExportSettingsFailure", "Export settings failed"), SNotificationItem::CS_Fail);
+	return FReply::Handled();
+}
+
+FReply SSettingsEditor::HandleImportButtonClicked()
+{
+	const TSharedPtr<SWindow> ParentWindow = FSlateApplication::Get().FindWidgetWindow(AsShared());
+	const void* ParentWindowHandle = ParentWindow.IsValid() && ParentWindow->GetNativeWindow().IsValid() ? ParentWindow->GetNativeWindow()->GetOSWindowHandle() : nullptr;
+
+	TArray<FString> OutFiles;
+	if (FDesktopPlatformModule::Get()->OpenFileDialog(ParentWindowHandle, 
+		LOCTEXT("ImportSettingsDialogTitle", "Import settings...").ToString(),
+		FPaths::GetPath(GEditorPerProjectIni), 
+		TEXT(""), 
+		TEXT("Config files (*.ini)|*.ini"), 
+		EFileDialogFlags::None, 
+		OutFiles))
+	{
+		const FString FilePath = OutFiles[0];
+
+		TArray<ISettingsCategoryPtr> Categories;
+		SettingsContainer->GetCategories(Categories);
+		for (const ISettingsCategoryPtr& Category : Categories)
+		{
+			TArray<ISettingsSectionPtr> Sections;
+			Category->GetSections(Sections);
+
+			for (const ISettingsSectionPtr& Section : Sections)
+			{
+				Section->Import(*FilePath);
+			}
+		}
+
+		ShowNotification(LOCTEXT("ImportSettingsSuccess", "Import settings succeeded"), SNotificationItem::CS_Success);
+	}
+	else
+	{
+		ShowNotification(LOCTEXT("ImportSettingsFailure", "Import settings failed"), SNotificationItem::CS_Fail);
+	}
+
+	return FReply::Handled();
+}
+
 bool SSettingsEditor::HandleSettingsViewEnabled() const
 {
 	ISettingsSectionPtr SelectedSection = Model->GetSelectedSection();
diff --git a/Engine/Source/Developer/SettingsEditor/Private/Widgets/SSettingsEditor.h b/Engine/Source/Developer/SettingsEditor/Private/Widgets/SSettingsEditor.h
--- a/Engine/Source/Developer/SettingsEditor/Private/Widgets/SSettingsEditor.h
+++ b/Engine/Source/Developer/SettingsEditor/Private/Widgets/SSettingsEditor.h
@@ -117,6 +117,12 @@ private:
 
 	/** Callback when the timer has been ticked to refresh the categories latently. */
 	bool UpdateCategoriesCallback(float InDeltaTime);
+
+	/** Callback for exporting all sections as a single ini. */
+	FReply HandleExportButtonClicked();
+
+	/** Callback for importing all sections from a single ini. */
+	FReply HandleImportButtonClicked();
 private:
 
 	/** Holds the vertical box for settings categories. */
diff --git a/Engine/Source/Editor/InputBindingEditor/Private/InputBindingEditorModule.cpp b/Engine/Source/Editor/InputBindingEditor/Private/InputBindingEditorModule.cpp
--- a/Engine/Source/Editor/InputBindingEditor/Private/InputBindingEditorModule.cpp
+++ b/Engine/Source/Editor/InputBindingEditor/Private/InputBindingEditorModule.cpp
@@ -465,12 +465,33 @@ private:
 	{
 		FInputBindingManager::Get().SaveInputBindings();
 		GConfig->Flush(false, GEditorKeyBindingsIni);
+
+		// If the user has no user-defined chords, the file won't be saved
+		if (!IFileManager::Get().FileExists(*GEditorKeyBindingsIni))
+		{
+			return false;
+		}
+
 		return BackupFile(GEditorKeyBindingsIni, Filename);
 	}
 
 	// Handles importing input bindings from a file
 	bool HandleInputBindingsImport(const FString& Filename)
 	{
+		if (!IFileManager::Get().FileExists(*Filename))
+		{
+			return false;
+		}
+
+		FConfigFile TempConfig;
+		TempConfig.Read(Filename);
+
+		// Don't bother prompting for a restart if the file doesn't actually contain any user-defined chords
+		if (!TempConfig.Contains(TEXT("UserDefinedChords")))
+		{
+			return false;
+		}
+
 		if(EAppReturnType::Ok == ShowRestartWarning(LOCTEXT("ImportKeyBindings_Title", "Import Key Bindings")))
 		{
 			FUnrealEdMisc::Get().SetConfigRestoreFilename(Filename, GEditorKeyBindingsIni);
