Fixed LOCTEXT conflict for Export/Import All buttons.

#rnx

diff --git a/Engine/Source/Developer/SettingsEditor/Private/Widgets/SSettingsEditor.cpp b/Engine/Source/Developer/SettingsEditor/Private/Widgets/SSettingsEditor.cpp
--- a/Engine/Source/Developer/SettingsEditor/Private/Widgets/SSettingsEditor.cpp
+++ b/Engine/Source/Developer/SettingsEditor/Private/Widgets/SSettingsEditor.cpp
@@ -148,7 +148,7 @@ void SSettingsEditor::Construct( const FArguments& InArgs, const ISettingsEditor
 						SNew(SButton)
 						.OnClicked(this, &SSettingsEditor::HandleExportButtonClicked)
 						.Text(LOCTEXT("ExportButtonText", "Export..."))
-						.ToolTipText(LOCTEXT("ExportButtonTooltip", "Export all editor settings to a file on your computer"))
+						.ToolTipText(LOCTEXT("ExportAllButtonTooltip", "Export all editor settings to a file on your computer"))
 					]
 					+ SHorizontalBox::Slot()
 					.AutoWidth()
@@ -159,7 +159,7 @@ void SSettingsEditor::Construct( const FArguments& InArgs, const ISettingsEditor
 						SNew(SButton)
 						.OnClicked(this, &SSettingsEditor::HandleImportButtonClicked)
 						.Text(LOCTEXT("ImportButtonText", "Import..."))
-						.ToolTipText(LOCTEXT("ImportButtonTooltip", "Import all editor settings from a file on your computer"))
+						.ToolTipText(LOCTEXT("ImportAllButtonTooltip", "Import all editor settings from a file on your computer"))
 					]
 				]
 				+ SVerticalBox::Slot()
