Enabled Reset to Default button on individual settings in Editor Preferences.

[RN]

diff --git a/Engine/Source/Developer/SettingsEditor/Private/Customizations/SettingsEditorCustomization.cpp b/Engine/Source/Developer/SettingsEditor/Private/Customizations/SettingsEditorCustomization.cpp
new file mode 100644
--- /dev/null
+++ b/Engine/Source/Developer/SettingsEditor/Private/Customizations/SettingsEditorCustomization.cpp
@@ -0,0 +1,280 @@
+// Copyright Epic Games, Inc. All Rights Reserved.
+
+#include "SettingsEditorCustomization.h"
+
+#include "DetailLayoutBuilder.h"
+#include "HAL/FileManager.h"
+#include "Misc/ConfigContext.h"
+#include "UObject/Package.h"
+
+TSharedRef<IDetailCustomization> FSettingsEditorCustomization::MakeInstance()
+{
+	return MakeShared<FSettingsEditorCustomization>();
+}
+
+namespace UE::Editor::Settings::Private
+{
+bool GetValueFromObject(TSharedPtr<IPropertyHandle> PropertyHandle, UObject* ValueSource, FString& OutValue)
+{
+	const FProperty* SettingsProperty = PropertyHandle->GetProperty();
+	const void* PropertyAddress = PropertyHandle->GetValueBaseAddress((uint8*)ValueSource);
+
+	constexpr uint32 PortFlags = PPF_None;
+	if (PropertyHandle->GetArrayIndex() == INDEX_NONE && SettingsProperty->ArrayDim > 1)
+	{
+		FArrayProperty::ExportTextInnerItem(OutValue, SettingsProperty, PropertyAddress, SettingsProperty->ArrayDim,
+			PropertyAddress, SettingsProperty->ArrayDim, nullptr, PortFlags);
+	}
+	else
+	{
+		PropertyHandle->GetProperty()->ExportTextItem_Direct(OutValue, PropertyAddress, PropertyAddress, nullptr, PortFlags);
+	}
+
+	const UEnum* Enum = nullptr;
+	int64 EnumValue = 0;
+	if (const FByteProperty* ByteProperty = CastField<FByteProperty>(SettingsProperty))
+	{
+		if (ByteProperty->Enum != nullptr)
+		{
+			Enum = ByteProperty->Enum.Get();
+			EnumValue = ByteProperty->GetPropertyValue(PropertyAddress);
+		}
+	}
+	else if (const FEnumProperty* EnumProperty = CastField<FEnumProperty>(SettingsProperty))
+	{
+		Enum = EnumProperty->GetEnum();
+		EnumValue = EnumProperty->GetUnderlyingProperty()->GetSignedIntPropertyValue(PropertyAddress);
+	}
+
+	if (Enum)
+	{
+		if (Enum->IsValidEnumValue(EnumValue))
+		{
+			OutValue = Enum->GetDisplayNameTextByValue(EnumValue).ToString();
+			if (OutValue.IsEmpty())
+			{
+				OutValue = Enum->GetNameStringByValue(EnumValue);
+			}
+		}
+		else
+		{
+			return false;
+		}
+	}
+
+	return true;
+}
+}
+
+void FSettingsEditorCustomization::CustomizeDetails(IDetailLayoutBuilder& InDetailLayoutBuilder)
+{
+	TArray<TWeakObjectPtr<UObject>> SettingsObjects;
+	InDetailLayoutBuilder.GetObjectsBeingCustomized(SettingsObjects);
+
+	for (const TWeakObjectPtr<UObject>& SettingsObject : SettingsObjects)
+	{
+		TStrongObjectPtr<UObject> PinnedSettingsObject = SettingsObject.Pin();
+		if (!PinnedSettingsObject)
+		{
+			continue;
+		}
+
+		ResetDefaultPropertyValues(InDetailLayoutBuilder, PinnedSettingsObject.Get());
+
+		const UClass* SettingsClass = SettingsObject->GetClass();
+		for (const FProperty* SettingsProperty = SettingsClass->PropertyLink; SettingsProperty; SettingsProperty = SettingsProperty->PropertyLinkNext)
+		{
+			if (!SettingsProperty->HasAnyPropertyFlags(CPF_Config | CPF_GlobalConfig))
+			{
+				continue;
+			}
+
+			TSharedRef<IPropertyHandle> PropertyHandle = InDetailLayoutBuilder.GetProperty(SettingsProperty->GetFName(), SettingsClass);
+			if (!PropertyHandle->IsValidHandle())
+			{
+				continue;
+			}
+
+			IDetailPropertyRow* PropertyRow = InDetailLayoutBuilder.EditDefaultProperty(PropertyHandle);
+			if (!PropertyRow)
+			{
+				continue;
+			}
+
+			PropertyRow->OverrideResetToDefault(FResetToDefaultOverride::Create(
+				FIsResetToDefaultVisible::CreateSP(this, &FSettingsEditorCustomization::IsResetToDefaultVisible),
+				FResetToDefaultHandler::CreateSP(this, &FSettingsEditorCustomization::OnResetToDefault)));
+		}
+	}
+}
+
+void FSettingsEditorCustomization::CustomizeDetails(const TSharedPtr<IDetailLayoutBuilder>& DetailBuilder)
+{
+	check(DetailBuilder)
+	DetailBuilderPtr = DetailBuilder;
+	FSettingsEditorCustomization::CustomizeDetails(*DetailBuilder);
+}
+
+void FSettingsEditorCustomization::OnResetToDefault(TSharedPtr<IPropertyHandle> InHandle)
+{
+	const FProperty* SettingsProperty = InHandle->GetProperty();
+	if (!SettingsProperty)
+	{
+		return;
+	}
+
+	TArray<UObject*> SettingsOuterObjects;
+	InHandle->GetOuterObjects(SettingsOuterObjects);
+	if (SettingsOuterObjects.IsEmpty())
+	{
+		return;
+	}
+
+	const FString* DefaultValuePtr = DefaultPropertyValues.Find(InHandle->GetProperty()->GetFName());
+	if (!DefaultValuePtr)
+	{
+		return;
+	}
+
+	const FString DefaultValue = *DefaultValuePtr;
+	InHandle->SetValueFromFormattedString(DefaultValue);
+}
+
+bool FSettingsEditorCustomization::IsResetToDefaultVisible(TSharedPtr<IPropertyHandle> InHandle)
+{
+	TArray<UObject*> SettingsOuterObjects;
+	InHandle->GetOuterObjects(SettingsOuterObjects);
+	if (SettingsOuterObjects.IsEmpty())
+	{
+		return false;
+	}
+
+	const FString* DefaultValue = DefaultPropertyValues.Find(InHandle->GetProperty()->GetFName());
+	if (!DefaultValue)
+	{
+		return false;
+	}
+
+	FString CurrentValue;
+	constexpr EPropertyPortFlags PortFlags = PPF_None;
+	if (InHandle->GetValueAsFormattedString(CurrentValue, PortFlags) != FPropertyAccess::Success)
+	{
+		return false;
+	}
+
+	return !CurrentValue.Equals(*DefaultValue, ESearchCase::CaseSensitive);
+}
+
+void FSettingsEditorCustomization::ResetDefaultPropertyValues(IDetailLayoutBuilder& InDetailLayoutBuilder, UObject* SettingsObject)
+{
+	DefaultPropertyValues.Reset();
+
+	UObject* TempDefaultObject = nullptr;
+	ON_SCOPE_EXIT
+	{
+		if (TempDefaultObject)
+		{
+			TempDefaultObject->ClearFlags(RF_ClassDefaultObject);
+			TempDefaultObject->Rename(nullptr, nullptr, REN_DoNotDirty | REN_DontCreateRedirectors | REN_NonTransactional);
+			TempDefaultObject->MarkAsGarbage();
+			TempDefaultObject = nullptr;
+		}
+	};
+
+	const UClass* SettingsClass = SettingsObject->GetClass();
+	const FString ConfigFilePath = SettingsClass->GetConfigName();
+	const FString SectionName = SettingsClass->GetPathName();
+
+	// First assume the current settings already match the defaults
+	UObject* ValueSource = SettingsObject;
+
+	// If the section exists in the current config, then defaults may have been overriden
+	const FConfigFile* ConfigFile = GConfig->Find(ConfigFilePath);
+	if (ConfigFile && ConfigFile->FindSection(SectionName))
+	{
+		constexpr bool bIncludeDotInExtension = true;
+		const FString ConfigDirectory = FPaths::GetPath(ConfigFilePath);
+		const FString BaseIniName = FPaths::GetBaseFilename(ConfigFilePath);
+		const FString ConfigExtension = FPaths::GetExtension(ConfigFilePath, bIncludeDotInExtension);
+
+		uint32 Counter = 1;
+		FString BackupConfigPath = FPaths::Combine(ConfigDirectory, FString::Printf(TEXT("%s_Backup%s"), *BaseIniName, *ConfigExtension));
+		while (IFileManager::Get().FileExists(*BackupConfigPath))
+		{
+			BackupConfigPath = FPaths::Combine(ConfigDirectory, FString::Printf(TEXT("%s_Backup%d%s"), *BaseIniName, Counter++, *ConfigExtension));
+		}
+
+		// Backup the current config in case something were to go wrong in this process
+		if (IFileManager::Get().Copy(*BackupConfigPath, *ConfigFilePath) != ECopyResult::COPY_OK)
+		{
+			return;
+		}
+
+		// Get rid of the current config values that may be overriding the defaults
+		GConfig->EmptySection(*SectionName, ConfigFilePath);
+		GConfig->Flush(false, ConfigFilePath);
+
+		// Reload the config, ensuring that we're only looking at default values
+		FConfigContext::ForceReloadIntoGConfig().Load(*BaseIniName);
+
+		UClass* ParentClass = SettingsClass->GetSuperClass();
+		UObject* ParentDefaultObject = nullptr;
+		if (ParentClass)
+		{
+			UObjectForceRegistration(ParentClass);
+			ParentDefaultObject = ParentClass->GetDefaultObject();
+			check(GConfig);
+			if (GEventDrivenLoaderEnabled && EVENT_DRIVEN_ASYNC_LOAD_ACTIVE_AT_RUNTIME)
+			{
+				check(ParentDefaultObject && !ParentDefaultObject->HasAnyFlags(RF_NeedLoad));
+			}
+		}
+
+		// Allocate a new, temporary default object
+		TempDefaultObject = StaticAllocateObject(SettingsClass, GetTransientPackage(), NAME_None, RF_ClassDefaultObject);
+		check(SettingsClass->GetDefaultObject() != TempDefaultObject); // Ensure we haven't been given the existing default object
+
+		EObjectInitializerOptions InitOptions = EObjectInitializerOptions::None;
+		if (!SettingsClass->HasAnyClassFlags(CLASS_Native | CLASS_Intrinsic))
+		{
+			// Blueprint CDOs have their properties always initialized.
+			InitOptions |= EObjectInitializerOptions::InitializeProperties;
+		}
+
+		// Initialize the temporary default object, loading the default config values
+		(*SettingsClass->ClassConstructor)(FObjectInitializer(TempDefaultObject, ParentDefaultObject, InitOptions));
+		ValueSource = TempDefaultObject;
+
+		// Restore the current config values
+		if (IFileManager::Get().Copy(*ConfigFilePath, *BackupConfigPath) != ECopyResult::COPY_OK
+			|| !FConfigContext::ForceReloadIntoGConfig().Load(*BaseIniName))
+		{
+			UE_LOG(LogConfig, Warning, TEXT("Failed to restore config %s while caching default values. Click Import in %s and select %s to restore previous values."), *ConfigFilePath, *SectionName, *BackupConfigPath);
+			return;
+		}
+
+		IFileManager::Get().Delete(*BackupConfigPath);
+	}
+
+	for (const FProperty* SettingsProperty = SettingsClass->PropertyLink; SettingsProperty; SettingsProperty = SettingsProperty->PropertyLinkNext)
+	{
+		if (!SettingsProperty->HasAnyPropertyFlags(CPF_Config | CPF_GlobalConfig))
+		{
+			continue;
+		}
+
+		TSharedRef<IPropertyHandle> PropertyHandle = InDetailLayoutBuilder.GetProperty(SettingsProperty->GetFName(), SettingsClass);
+		if (!PropertyHandle->IsValidHandle())
+		{
+			continue;
+		}
+
+		FString DefaultValue;
+		if (!UE::Editor::Settings::Private::GetValueFromObject(PropertyHandle, ValueSource, DefaultValue))
+		{
+			continue;
+		}
+
+		DefaultPropertyValues.Add(PropertyHandle->GetProperty()->GetFName(), DefaultValue);
+	}
+}
diff --git a/Engine/Source/Developer/SettingsEditor/Private/Customizations/SettingsEditorCustomization.h b/Engine/Source/Developer/SettingsEditor/Private/Customizations/SettingsEditorCustomization.h
new file mode 100644
--- /dev/null
+++ b/Engine/Source/Developer/SettingsEditor/Private/Customizations/SettingsEditorCustomization.h
@@ -0,0 +1,29 @@
+// Copyright Epic Games, Inc. All Rights Reserved.
+
+#pragma once
+
+#include "IDetailCustomization.h"
+#include "PropertyHandle.h"
+
+//////////////////////////////////////////////////////////////////////////
+// FSettingsEditorCustomization
+
+struct FInstancedPropertyBag;
+
+class FSettingsEditorCustomization : public IDetailCustomization
+{
+public:
+	static TSharedRef<IDetailCustomization> MakeInstance();
+	virtual void CustomizeDetails(IDetailLayoutBuilder& InDetailLayoutBuilder) override;
+	virtual void CustomizeDetails(const TSharedPtr<IDetailLayoutBuilder>& DetailBuilder) override;
+
+private:
+	void OnResetToDefault(TSharedPtr<IPropertyHandle> InHandle);
+	bool IsResetToDefaultVisible(TSharedPtr<IPropertyHandle> InHandle);
+
+	// Cache the latest default property values
+	void ResetDefaultPropertyValues(IDetailLayoutBuilder& InDetailLayoutBuilder, UObject* SettingsObject);
+
+	TMap<FName, FString> DefaultPropertyValues;
+	TWeakPtr<IDetailLayoutBuilder> DetailBuilderPtr;
+};
diff --git a/Engine/Source/Developer/SettingsEditor/Private/Widgets/SSettingsEditor.cpp b/Engine/Source/Developer/SettingsEditor/Private/Widgets/SSettingsEditor.cpp
--- a/Engine/Source/Developer/SettingsEditor/Private/Widgets/SSettingsEditor.cpp
+++ b/Engine/Source/Developer/SettingsEditor/Private/Widgets/SSettingsEditor.cpp
@@ -25,6 +25,7 @@
 #include "HAL/PlatformFileManager.h"
 #include "HAL/PlatformFile.h"
 #include "Containers/Ticker.h"
+#include "Customizations/SettingsEditorCustomization.h"
 #include "Framework/Notifications/NotificationManager.h"
 #include "HAL/FileManager.h"
 #include "Misc/ConfigCacheIni.h"
@@ -78,6 +79,7 @@ void SSettingsEditor::Construct( const FArguments& InArgs, const ISettingsEditor
 	SettingsView = FModuleManager::GetModuleChecked<FPropertyEditorModule>("PropertyEditor").CreateDetailView(DetailsViewArgs);
 	SettingsView->SetVisibility(TAttribute<EVisibility>::CreateSP(this, &SSettingsEditor::HandleSettingsViewVisibility));
 	SettingsView->SetIsPropertyEditingEnabledDelegate(FIsPropertyEditingEnabled::CreateSP(this, &SSettingsEditor::HandleSettingsViewEnabled));
+	SettingsView->SetGenericLayoutDetailsDelegate(FOnGetDetailCustomizationInstance::CreateStatic(&FSettingsEditorCustomization::MakeInstance));
 
 	TSharedPtr<FSettingsDetailRootObjectCustomization> RootObjectCustomization = MakeShareable(new FSettingsDetailRootObjectCustomization(Model, SettingsView.ToSharedRef()));
 	RootObjectCustomization->Initialize();
