Fix settings being reset to default if editor preferences left open on close editor
#rnx

diff --git a/Engine/Source/Developer/SettingsEditor/Private/Customizations/SettingsEditorCustomization.cpp b/Engine/Source/Developer/SettingsEditor/Private/Customizations/SettingsEditorCustomization.cpp
--- a/Engine/Source/Developer/SettingsEditor/Private/Customizations/SettingsEditorCustomization.cpp
+++ b/Engine/Source/Developer/SettingsEditor/Private/Customizations/SettingsEditorCustomization.cpp
@@ -79,8 +79,6 @@ void FSettingsEditorCustomization::CustomizeDetails(IDetailLayoutBuilder& InDeta
 			continue;
 		}
 
-		ResetDefaultPropertyValues(InDetailLayoutBuilder, PinnedSettingsObject.Get());
-
 		const UClass* SettingsClass = SettingsObject->GetClass();
 		for (const FProperty* SettingsProperty = SettingsClass->PropertyLink; SettingsProperty; SettingsProperty = SettingsProperty->PropertyLinkNext)
 		{
