EditorUsability : SettingsEditor
- Added context menu entry in editor preference and project settings to copy associated CVar to the clipboard if any
- Added row extension button to quickly spot properties with a CVar and copy them to clipboard
- Increased default extension column width to display button

#Synced-CL 48229788

diff --git a/Engine/Source/Developer/SettingsEditor/Private/SettingsEditorModule.cpp b/Engine/Source/Developer/SettingsEditor/Private/SettingsEditorModule.cpp
--- a/Engine/Source/Developer/SettingsEditor/Private/SettingsEditorModule.cpp
+++ b/Engine/Source/Developer/SettingsEditor/Private/SettingsEditorModule.cpp
@@ -1,27 +1,29 @@
 // Copyright Epic Games, Inc. All Rights Reserved.
 
-#include "CoreMinimal.h"
+#include "Customizations/PerPlatformSettingsCustomization.h"
+#include "DetailRowMenuContext.h"
+#include "Engine/DeveloperSettings.h"
+#include "Engine/PlatformSettings.h"
+#include "Framework/Notifications/NotificationManager.h"
+#include "IDetailTreeNode.h"
+#include "ISettingsContainer.h"
+#include "ISettingsEditorModel.h"
+#include "ISettingsEditorModule.h"
+#include "ISettingsModule.h"
+#include "Models/SettingsEditorModel.h"
 #include "Modules/ModuleManager.h"
+#include "PropertyEditorClipboard.h"
+#include "PropertyEditorDelegates.h"
+#include "PropertyEditorModule.h"
+#include "ToolMenus.h"
 #include "UObject/UObjectHash.h"
 #include "UObject/UObjectIterator.h"
+#include "UObject/UObjectThreadContext.h"
 #include "Widgets/DeclarativeSyntaxSupport.h"
-#include "ISettingsContainer.h"
-#include "ISettingsEditorModel.h"
-#include "Models/SettingsEditorModel.h"
 #include "Widgets/SWidget.h"
-#include "Framework/Notifications/NotificationManager.h"
 #include "Widgets/Notifications/SNotificationList.h"
 #include "Widgets/SSettingsEditor.h"
 
-#include "ISettingsEditorModule.h"
-#include "ISettingsModule.h"
-#include "Engine/DeveloperSettings.h"
-#include "Engine/PlatformSettings.h"
-#include "PropertyEditorModule.h"
-#include "PropertyEditorDelegates.h"
-#include "Customizations/PerPlatformSettingsCustomization.h"
-#include "UObject/UObjectThreadContext.h"
-
 #define LOCTEXT_NAMESPACE "SSettingsEditor"
 
 /** Holds auto discovered settings information so that they can be unloaded automatically when refreshing. */
@@ -32,7 +34,6 @@ struct FRegisteredSettings
 	FName SectionName;
 };
 
-
 /** Manages the notification for when the application needs to be restarted due to a settings change */
 class FApplicationRestartRequiredNotification
 {
@@ -143,6 +144,8 @@ public:
 		);
 
 		FModuleManager::Get().OnModulesChanged().AddRaw(this, &FSettingsEditorModule::ModulesChangesCallback);
+		
+		RegisterSettingsActions();
 	}
 
 	virtual void ShutdownModule() override
@@ -155,6 +158,8 @@ public:
 		}
 
 		FModuleManager::Get().OnModulesChanged().RemoveAll(this);
+		
+		UnregisterSettingsActions();
 	}
 
 	// ISettingsEditorModule interface
@@ -338,6 +343,152 @@ private:
 	}
 
 private:
+	void RegisterSettingsActions()
+	{
+		if (FPropertyEditorModule* PropertyModule = FModuleManager::GetModulePtr<FPropertyEditorModule>("PropertyEditor"))
+		{
+			PropertyModule->GetGlobalRowExtensionDelegate().AddStatic(&FSettingsEditorModule::OnExtendRowExtensionButtons);
+		}
+	}
+	
+	void UnregisterSettingsActions() const
+	{
+		if (FPropertyEditorModule* PropertyModule = FModuleManager::GetModulePtr<FPropertyEditorModule>("PropertyEditor"))
+		{
+			PropertyModule->GetGlobalRowExtensionDelegate().RemoveAll(this);
+		}
+	}
+
+	static void ExtendRowContextMenu()
+	{
+		static bool bMenuExtended = false;
+
+		if (!bMenuExtended)
+		{
+			UToolMenus* Menus = UToolMenus::Get();
+
+			checkf(Menus, TEXT("ToolMenus not available"));
+
+			if (UToolMenu* ContextMenu = Menus->FindMenu(UE::PropertyEditor::RowContextMenuName))
+			{
+				ContextMenu->AddDynamicSection(
+					TEXT("FillSettingsRowContextSection"), 
+					FNewToolMenuDelegate::CreateStatic(&FSettingsEditorModule::FillSettingsRowContextSection)
+				);
+				
+				bMenuExtended =  true;
+			}
+		}
+	}
+
+	static void FillSettingsRowContextSection(UToolMenu* InToolMenu)
+	{
+		if (!InToolMenu)
+		{
+			return;
+		}
+
+		UDetailRowMenuContext* RowContext = InToolMenu->FindContext<UDetailRowMenuContext>();
+		if (!RowContext 
+			|| RowContext->PropertyHandles.IsEmpty() 
+			|| !RowContext->PropertyHandles[0].IsValid())
+		{
+			return;
+		}
+
+		// Only fill menu with a property linked to a CVar
+		const FString CVar = GetCVarFromPropertyHandle(RowContext->PropertyHandles[0]);
+		if (CVar.IsEmpty())
+		{
+			return;
+		}
+
+		InToolMenu->AddMenuEntry(
+			TEXT("CopyCVarName"),
+			FToolMenuEntry::InitMenuEntry(
+				TEXT("CopyCvarName"),
+				LOCTEXT("CopyCVarName", "Copy CVar Name"),
+				FText::Format(LOCTEXT("CopyPropertyCVarNameToolTip", "Copy the console variable name of this property to the system clipboard:\n{0}"), FText::FromString(CVar)),
+				FSlateIcon(FCoreStyle::Get().GetStyleSetName(), "GenericCommands.Copy"),
+				FUIAction(
+					FExecuteAction::CreateLambda(&FSettingsEditorModule::CopyCVarToClipboard, CVar)
+				)
+			)
+		);
+	}
+
+	static void OnExtendRowExtensionButtons(const FOnGenerateGlobalRowExtensionArgs& InArgs, TArray<FPropertyRowExtensionButton>& OutExtensionButtons)
+	{
+		if (!InArgs.PropertyHandle.IsValid() || !InArgs.PropertyHandle->IsValidHandle())
+		{
+			return;
+		}
+
+		const TSharedPtr<IDetailTreeNode> DetailNode = InArgs.OwnerTreeNode.Pin();
+		if (!DetailNode.IsValid())
+		{
+			return;
+		}
+
+		const TSharedPtr<IDetailsView> DetailsView = DetailNode->GetNodeDetailsViewSharedPtr();
+		if (!DetailsView)
+		{
+			return;
+		}
+		
+		// Only extend rows that belong to settings editor
+		if (DetailsView->GetIdentifier() != TEXT("SSettingsEditor"))
+		{
+			return;
+		}
+
+		// Only extend rows with a property linked to a CVar
+		const FString CVarMetaData = GetCVarFromPropertyHandle(InArgs.PropertyHandle);
+		if (CVarMetaData.IsEmpty())
+		{
+			return;
+		}
+		
+		// Extend context menu here to ensure menu is registered and ready, this is only done once
+		ExtendRowContextMenu();
+
+		FPropertyRowExtensionButton& ExtensionButton = OutExtensionButtons.Add_GetRef(FPropertyRowExtensionButton());
+		ExtensionButton.Icon = FSlateIcon("OutputLogStyle", "DebugConsole.Icon");
+		ExtensionButton.Label = LOCTEXT("ExtensionCopyCVarName", "Copy CVar Name");
+		ExtensionButton.ToolTip = FText::Format(
+			LOCTEXT("ExtensionCopyCVarNameTooltip", "Copy CVar linked to this property row in the clipboard:\n{0}"),
+			FText::FromString(CVarMetaData)
+		);
+		ExtensionButton.UIAction = FUIAction(
+			FExecuteAction::CreateLambda(&FSettingsEditorModule::CopyCVarToClipboard, CVarMetaData)
+		);
+	}
+	
+	static FString GetCVarFromPropertyHandle(const TSharedPtr<IPropertyHandle>& InPropertyHandle)
+	{
+		if (InPropertyHandle)
+		{
+			const FString ConsoleVariableMeta = TEXT("ConsoleVariable");
+			if (InPropertyHandle->HasMetaData(*ConsoleVariableMeta))
+			{
+				return InPropertyHandle->GetMetaData(*ConsoleVariableMeta);
+			}
+		}
+
+		return FString();
+	}
+	
+	static void CopyCVarToClipboard(const FString& InCVar)
+	{
+		FPropertyEditorClipboard::ClipboardCopy(*InCVar);
+
+		// Notify of the action
+		FNotificationInfo NotificationInfo(LOCTEXT("CopyCVarNotification", "CVar copied to clipboard"));
+		NotificationInfo.ExpireDuration = 1.f;
+		NotificationInfo.bFireAndForget = true;
+		NotificationInfo.Image = FAppStyle::GetBrush("Icons.SuccessWithColor");
+		FSlateNotificationManager::Get().AddNotification(NotificationInfo);
+	}
 
 	FApplicationRestartRequiredNotification ApplicationRestartRequiredNotification;
 
diff --git a/Engine/Source/Developer/SettingsEditor/Private/Widgets/SSettingsEditor.cpp b/Engine/Source/Developer/SettingsEditor/Private/Widgets/SSettingsEditor.cpp
--- a/Engine/Source/Developer/SettingsEditor/Private/Widgets/SSettingsEditor.cpp
+++ b/Engine/Source/Developer/SettingsEditor/Private/Widgets/SSettingsEditor.cpp
@@ -74,6 +74,8 @@ void SSettingsEditor::Construct( const FArguments& InArgs, const ISettingsEditor
 		DetailsViewArgs.bCustomNameAreaLocation = true;
 		DetailsViewArgs.bCustomFilterAreaLocation = true;
 		DetailsViewArgs.NameAreaSettings = FDetailsViewArgs::HideNameArea;
+		DetailsViewArgs.RightColumnMinWidth = 50.f;
+		DetailsViewArgs.ViewIdentifier = TEXT("SSettingsEditor");
 	}
 
 	SettingsView = FModuleManager::GetModuleChecked<FPropertyEditorModule>("PropertyEditor").CreateDetailView(DetailsViewArgs);
diff --git a/Engine/Source/Developer/SettingsEditor/SettingsEditor.Build.cs b/Engine/Source/Developer/SettingsEditor/SettingsEditor.Build.cs
--- a/Engine/Source/Developer/SettingsEditor/SettingsEditor.Build.cs
+++ b/Engine/Source/Developer/SettingsEditor/SettingsEditor.Build.cs
@@ -23,7 +23,7 @@ public class SettingsEditor : ModuleRules
 				"Analytics",
                 "CoreUObject",
 				"DesktopPlatform",
-                
+				"DeveloperSettings",
 				"Engine",
                 "InputCore",
 				"PropertyEditor",
@@ -31,7 +31,7 @@ public class SettingsEditor : ModuleRules
 				"Slate",
 				"SlateCore",
 				"SourceControl",
-				"DeveloperSettings"
+				"ToolMenus"
             }
         );
 	}
