EditorUsability : SettingsEditor
- Added ability to search for CVAR/KEYBIND in Editor Preference or Project Settings, the details view search widget is replaced by a custom search widget with suggestions, I don't control how the search is performed since it's done by the details panel, I can only suggest and show what can be searched to the user and autocomplete. Suggestions and results are based on the current context, meaning it will not list cvar that are editor in project settings or editor keybinds in project settings. This is mimicking advanced search synthax without using it since it would require refactoring the entire details panel search capability and this is not the goal of this change. This system does allow for additional search terms to be registered
- Added quick filters checkbox buttons that filter result
- Simplifying keybind list by removing description based on command tooltip
- DetailsView changes were kept minimal, only what is needed has been added

#Synced-CL 49549735

diff --git a/Engine/Source/Developer/SettingsEditor/Private/Customizations/SettingsEditorCustomization.cpp b/Engine/Source/Developer/SettingsEditor/Private/Customizations/SettingsEditorCustomization.cpp
--- a/Engine/Source/Developer/SettingsEditor/Private/Customizations/SettingsEditorCustomization.cpp
+++ b/Engine/Source/Developer/SettingsEditor/Private/Customizations/SettingsEditorCustomization.cpp
@@ -3,8 +3,11 @@
 #include "SettingsEditorCustomization.h"
 
 #include "DetailLayoutBuilder.h"
+#include "DetailWidgetRow.h"
+#include "ISettingsEditorModule.h"
 #include "HAL/FileManager.h"
 #include "Misc/ConfigContext.h"
+#include "Modules/ModuleManager.h"
 #include "UObject/Package.h"
 
 TSharedRef<IDetailCustomization> FSettingsEditorCustomization::MakeInstance()
@@ -70,6 +73,14 @@ void FSettingsEditorCustomization::CustomizeDetails(IDetailLayoutBuilder& InDeta
 {
 	TArray<TWeakObjectPtr<UObject>> SettingsObjects;
 	InDetailLayoutBuilder.GetObjectsBeingCustomized(SettingsObjects);
+	
+	ISettingsEditorModule* SettingsEditorModule = FModuleManager::GetModulePtr<ISettingsEditorModule>(TEXT("SettingsEditor"));
+	const FString SearchKey = TEXT("cvar");
+	FName SearchContext = NAME_None;
+	if (const TSharedPtr<IDetailsView> DetailsView = InDetailLayoutBuilder.GetDetailsViewSharedPtr())
+	{
+		SearchContext = DetailsView->GetIdentifier();
+	}
 
 	for (const TWeakObjectPtr<UObject>& SettingsObject : SettingsObjects)
 	{
@@ -98,10 +109,49 @@ void FSettingsEditorCustomization::CustomizeDetails(IDetailLayoutBuilder& InDeta
 			{
 				continue;
 			}
-
-			PropertyRow->OverrideResetToDefault(FResetToDefaultOverride::Create(
-				FIsResetToDefaultVisible::CreateSP(this, &FSettingsEditorCustomization::IsResetToDefaultVisible),
-				FResetToDefaultHandler::CreateSP(this, &FSettingsEditorCustomization::OnResetToDefault)));
+			
+			const FString& CVar = SettingsProperty->GetMetaData(TEXT("ConsoleVariable"));
+			if (!CVar.IsEmpty())
+			{
+				FDetailWidgetRow& Row = PropertyRow->CustomWidget(true);
+				Row.NameContent()
+				[
+					PropertyHandle->CreatePropertyNameWidget()
+				];
+
+				Row.ValueContent()
+				[
+					PropertyHandle->CreatePropertyValueWidget()
+				];
+
+				// Allow searching of cvar in settings editor (for i.e. cvar=r.vsync)
+				Row.FilterString(FText::FromString(TEXT("cvar=") + CVar));
+
+				// Override reset to default
+				Row.OverrideResetToDefault(FResetToDefaultOverride::Create(
+					FIsResetToDefaultVisible::CreateSP(this, &FSettingsEditorCustomization::IsResetToDefaultVisible),
+					FResetToDefaultHandler::CreateSP(this, &FSettingsEditorCustomization::OnResetToDefault))
+				);
+				
+				// Register search key and value for settings editor suggestions
+				if (SettingsEditorModule)
+				{
+					ISettingsEditorModule::FSearchTerm SearchTerm
+					{
+						.ContainerName = SearchContext,
+						.Key = SearchKey,
+						.Value = CVar
+					};
+					SettingsEditorModule->RegisterSearchTerm(SearchTerm);
+				}
+			}
+			else
+			{
+				PropertyRow->OverrideResetToDefault(FResetToDefaultOverride::Create(
+					FIsResetToDefaultVisible::CreateSP(this, &FSettingsEditorCustomization::IsResetToDefaultVisible),
+					FResetToDefaultHandler::CreateSP(this, &FSettingsEditorCustomization::OnResetToDefault))
+				);
+			}
 		}
 	}
 }
diff --git a/Engine/Source/Developer/SettingsEditor/Private/SettingsEditorModule.cpp b/Engine/Source/Developer/SettingsEditor/Private/SettingsEditorModule.cpp
--- a/Engine/Source/Developer/SettingsEditor/Private/SettingsEditorModule.cpp
+++ b/Engine/Source/Developer/SettingsEditor/Private/SettingsEditorModule.cpp
@@ -196,6 +196,31 @@ public:
 	{
 		ShouldRegisterSettingsDelegate = MoveTemp(InShouldRegisterSettingDelegate);
 	}
+	
+	virtual void RegisterSearchTerm(const FSearchTerm& InSearchTerm) override
+	{
+		if (InSearchTerm.ContainerName.IsNone())
+		{
+			return;
+		}
+
+		SearchTerms
+			.FindOrAdd(InSearchTerm.ContainerName)
+			.FindOrAdd(InSearchTerm.Key.ToLower())
+			.Add(InSearchTerm.Value);
+
+		OnSearchTermsChangedDelegate.Broadcast(InSearchTerm);
+	}
+	
+	virtual FOnSearchTermsChanged::RegistrationType& OnSearchTermsChanged() override
+	{
+		return OnSearchTermsChangedDelegate;
+	}
+	
+	virtual const TMap<FString, TSet<FString>>* FindSearchTerms(FName InContainerName) const override
+	{
+		return SearchTerms.Find(InContainerName);
+	}
 
 private:
 
@@ -437,7 +462,8 @@ private:
 		}
 		
 		// Only extend rows that belong to settings editor
-		if (DetailsView->GetIdentifier() != TEXT("SSettingsEditor"))
+		const FName ViewIdentifier = DetailsView->GetIdentifier();
+		if (ViewIdentifier != TEXT("Project") && ViewIdentifier != TEXT("Editor"))
 		{
 			return;
 		}
@@ -506,6 +532,12 @@ private:
 
 	/** Delegate called to decide if a settings object should be registered with the module. */
 	FShouldRegisterSettingsDelegate ShouldRegisterSettingsDelegate;
+	
+	/** Used to store Context->Key->Values for search suggestions */
+	TMap<FName, TMap<FString, TSet<FString>>> SearchTerms;
+	
+	/** Delegate called when search suggestions changed */
+	FOnSearchTermsChanged OnSearchTermsChangedDelegate;
 };
 
 
diff --git a/Engine/Source/Developer/SettingsEditor/Private/Widgets/SSettingsEditor.cpp b/Engine/Source/Developer/SettingsEditor/Private/Widgets/SSettingsEditor.cpp
--- a/Engine/Source/Developer/SettingsEditor/Private/Widgets/SSettingsEditor.cpp
+++ b/Engine/Source/Developer/SettingsEditor/Private/Widgets/SSettingsEditor.cpp
@@ -1,6 +1,8 @@
 // Copyright Epic Games, Inc. All Rights Reserved.
 
 #include "Widgets/SSettingsEditor.h"
+
+#include "Editor/EditorWidgets/Public/SAssetSearchBox.h"
 #include "UObject/UnrealType.h"
 #include "Misc/Paths.h"
 #include "Modules/ModuleManager.h"
@@ -43,6 +45,35 @@ SSettingsEditor::~SSettingsEditor()
 	FInternationalization::Get().OnCultureChanged().RemoveAll(this);
 }
 
+void SSettingsEditor::HandleSearchTermsChanged(const ISettingsEditorModule::FSearchTerm& InNewSearchTerm)
+{
+	if (InNewSearchTerm.ContainerName != SettingsContainer->GetName() || QuickSearchFiltersAdded.Contains(InNewSearchTerm.Key))
+	{
+		return;
+	}
+
+	QuickSearchFiltersAdded.Add(InNewSearchTerm.Key);
+
+	const FString SearchKey = InNewSearchTerm.Key;
+	SearchFiltersBox->AddSlot()
+	.Padding(0.f)
+    [
+    	SNew(SBox)
+    	.Padding(0.f, 8.f, 8.f, 8.f)
+    	[
+    		SNew(SCheckBox)
+    		.Style(FAppStyle::Get(), "DetailsView.SectionButton")
+    		.OnCheckStateChanged(this, &SSettingsEditor::OnSearchFilterButtonClicked, SearchKey)
+    		.IsChecked(this, &SSettingsEditor::IsSearchFilterButtonActive, SearchKey)
+    		[
+    			SNew(STextBlock)
+				.Text(FText::FromString(SearchKey))
+    			.ToolTipText(LOCTEXT("SearchFilterToolTip", "Toggles this search filter on/off"))
+    		]
+    	]
+    ];
+}
+
 
 /* SSettingsEditor interface
  *****************************************************************************/
@@ -70,12 +101,14 @@ void SSettingsEditor::Construct( const FArguments& InArgs, const ISettingsEditor
 		DetailsViewArgs.bShowKeyablePropertiesOption = false;
 		DetailsViewArgs.bShowHiddenPropertiesWhilePlayingOption = false;
 		DetailsViewArgs.bShowPropertyMatrixButton = false;
+		DetailsViewArgs.bShowSectionSelector = false;
 		DetailsViewArgs.bAllowMultipleTopLevelObjects = true;
 		DetailsViewArgs.bCustomNameAreaLocation = true;
 		DetailsViewArgs.bCustomFilterAreaLocation = true;
 		DetailsViewArgs.NameAreaSettings = FDetailsViewArgs::HideNameArea;
 		DetailsViewArgs.RightColumnMinWidth = 50.f;
-		DetailsViewArgs.ViewIdentifier = TEXT("SSettingsEditor");
+		// Used to identify details view
+		DetailsViewArgs.ViewIdentifier = SettingsContainer ? SettingsContainer->GetName() : TEXT("SSettingsEditor");
 	}
 
 	SettingsView = FModuleManager::GetModuleChecked<FPropertyEditorModule>("PropertyEditor").CreateDetailView(DetailsViewArgs);
@@ -140,8 +173,19 @@ void SSettingsEditor::Construct( const FArguments& InArgs, const ISettingsEditor
 					SNew(SHorizontalBox)
 					+ SHorizontalBox::Slot()
 					.FillWidth(1.0f)
+					.VAlign(VAlign_Center)
+					.Padding(0.f, 0.f, 0.f, 0.f)
 					[
-						SettingsView->GetFilterAreaWidget().ToSharedRef()
+						SAssignNew(SearchBox, SAssetSearchBox)
+						.HintText(LOCTEXT("FilterSearch", "Search..."))
+						.ShowSearchHistory(true)
+						.ToolTipText(this, &SSettingsEditor::GetSearchToolTipText)
+						.OnTextChanged(this, &SSettingsEditor::OnSearchTextChanged)
+						.OnTextCommitted(this, &SSettingsEditor::OnSearchTextCommitted)
+						.OnAssetSearchBoxSuggestionFilter(this, &SSettingsEditor::OnAssetSearchSuggestionFilter)
+						.OnAssetSearchBoxSuggestionChosen(this, &SSettingsEditor::OnAssetSearchSuggestionChosen)
+						.DelayChangeNotificationsWhileTyping(true)
+						.AddMetaData<FTagMetaData>(FTagMetaData(TEXT("SettingsEditorSearchBox")))
 					]
 					+ SHorizontalBox::Slot()
 					.AutoWidth()
@@ -167,6 +211,13 @@ void SSettingsEditor::Construct( const FArguments& InArgs, const ISettingsEditor
 					]
 				]
 				+ SVerticalBox::Slot()
+				.AutoHeight()
+				[
+					SAssignNew(SearchFiltersBox, SScrollBox)
+                    .Orientation(EOrientation::Orient_Horizontal)
+                    .ScrollBarThickness(FVector2D(1.f))
+				]
+				+ SVerticalBox::Slot()
 				.FillHeight(1.0f)
 				[
 					// settings area
@@ -187,6 +238,12 @@ void SSettingsEditor::Construct( const FArguments& InArgs, const ISettingsEditor
 	FInternationalization::Get().OnCultureChanged().AddSP(this, &SSettingsEditor::HandleCultureChanged);
 	Model->OnSelectionChanged().AddSP(this, &SSettingsEditor::HandleModelSelectionChanged);
 	SettingsContainer->OnCategoryModified().AddSP(this, &SSettingsEditor::HandleSettingsContainerCategoryModified);
+	
+	// Update quick filters keys based on new search terms registered
+	if (ISettingsEditorModule* SettingsEditorModule = FModuleManager::GetModulePtr<ISettingsEditorModule>(TEXT("SettingsEditor")))
+	{
+		SettingsEditorModule->OnSearchTermsChanged().AddSP(this, &SSettingsEditor::HandleSearchTermsChanged);
+	}
 
 	ReloadCategories();
 }
@@ -756,6 +813,124 @@ FReply SSettingsEditor::HandleImportButtonClicked()
 	return FReply::Handled();
 }
 
+FText SSettingsEditor::GetSearchToolTipText() const
+{
+	FText ToolTipText = LOCTEXT("FilterSearchTooltip", "Type here to search...");
+
+	if (const ISettingsEditorModule* SettingsEditorModule = FModuleManager::GetModulePtr<ISettingsEditorModule>(TEXT("SettingsEditor")))
+	{
+		if (SettingsContainer)
+		{
+			if (const TMap<FString, TSet<FString>>* SearchTerms = SettingsEditorModule->FindSearchTerms(SettingsContainer->GetName()))
+			{
+				TArray<FString> Keys;
+				SearchTerms->GenerateKeyArray(Keys);
+				ToolTipText = FText::Format(INVTEXT("{0} ({1})"), ToolTipText, FText::FromString(FString::Join(Keys, TEXT(", "))));
+			}
+		}
+	}
+
+	return ToolTipText;
+}
+
+void SSettingsEditor::OnSearchTextChanged(const FText& InText)
+{
+	if (SettingsView.IsValid())
+	{
+		SettingsView->SetSearchText(InText);
+	}
+}
+
+void SSettingsEditor::OnSearchTextCommitted(const FText& InText, ETextCommit::Type InType)
+{
+	if (SettingsView.IsValid() && InType == ETextCommit::Type::OnCleared)
+	{
+		SettingsView->ClearSearch();
+	}
+}
+
+void SSettingsEditor::OnAssetSearchSuggestionFilter(const FText& InText, TArray<FAssetSearchBoxSuggestion>& OutSuggestion, FText& OutText)
+{
+	const ISettingsEditorModule* SettingsEditorModule = FModuleManager::GetModulePtr<ISettingsEditorModule>(TEXT("SettingsEditor"));
+	if (!SettingsEditorModule || !SettingsContainer)
+	{
+		return;
+	}
+
+	const TMap<FString, TSet<FString>>* SearchTerms = SettingsEditorModule->FindSearchTerms(SettingsContainer->GetName());
+	if (!SearchTerms)
+	{
+		return;
+	}
+
+	const FString ActiveSearch = InText.ToString();
+	const FText SettingsCategory = LOCTEXT("SettingsEditor", "Settings");
+	for (const TPair<FString, TSet<FString>>& SearchTerm : *SearchTerms)
+	{
+		if (ActiveSearch.Contains(SearchTerm.Key) || SearchTerm.Key.Contains(ActiveSearch))
+		{
+			FString SearchSuggestion = SearchTerm.Key + TEXT("=");
+			OutSuggestion.Add(FAssetSearchBoxSuggestion(SearchSuggestion, FText::FromString(SearchSuggestion), SettingsCategory));
+			FString SearchCopy = ActiveSearch;
+			SearchCopy.RemoveFromStart(SearchSuggestion);
+
+			for (const FString& SearchValue : SearchTerm.Value)
+			{
+				if (SearchCopy.IsEmpty() || SearchValue.StartsWith(SearchCopy))
+				{
+					SearchSuggestion = SearchTerm.Key + TEXT("=") + SearchValue;
+					OutSuggestion.Add(FAssetSearchBoxSuggestion(SearchSuggestion, FText::FromString(SearchSuggestion), SettingsCategory));
+				}
+			}
+		}
+	}
+}
+
+FText SSettingsEditor::OnAssetSearchSuggestionChosen(const FText& InText, const FString& InSuggestion)
+{
+	FString CurrentSearch = InText.ToString();
+	if (!InSuggestion.Contains(CurrentSearch))
+	{
+		for (const auto& Char : InSuggestion)
+		{
+			if (!CurrentSearch.IsEmpty() && CurrentSearch[0] == Char)
+			{
+				CurrentSearch.RemoveAt(0);
+			}
+			else
+			{
+				break;
+			}
+		}
+		
+		CurrentSearch = InSuggestion + CurrentSearch;
+	}
+	else
+	{
+		CurrentSearch = InSuggestion;
+	}
+
+	return FText::FromString(CurrentSearch);
+}
+
+ECheckBoxState SSettingsEditor::IsSearchFilterButtonActive(FString InFilter) const
+{
+	const FString ActiveSearch = SearchBox->GetText().ToString();
+	return ActiveSearch.StartsWith(InFilter) ? ECheckBoxState::Checked : ECheckBoxState::Unchecked;
+}
+
+void SSettingsEditor::OnSearchFilterButtonClicked(ECheckBoxState InState, FString InFilter)
+{
+	if (InState == ECheckBoxState::Checked)
+	{
+		SearchBox->SetText(FText::FromString(InFilter + TEXT("=")));
+	}
+	else
+	{
+		SearchBox->SetText(FText::GetEmpty());
+	}
+}
+
 bool SSettingsEditor::HandleSettingsViewEnabled() const
 {
 	ISettingsSectionPtr SelectedSection = Model->GetSelectedSection();
diff --git a/Engine/Source/Developer/SettingsEditor/Private/Widgets/SSettingsEditor.h b/Engine/Source/Developer/SettingsEditor/Private/Widgets/SSettingsEditor.h
--- a/Engine/Source/Developer/SettingsEditor/Private/Widgets/SSettingsEditor.h
+++ b/Engine/Source/Developer/SettingsEditor/Private/Widgets/SSettingsEditor.h
@@ -3,19 +3,23 @@
 #pragma once
 
 #include "CoreMinimal.h"
-#include "Misc/NotifyHook.h"
+#include "ISettingsCategory.h"
+#include "ISettingsContainer.h"
+#include "ISettingsEditorModel.h"
+#include "ISettingsEditorModule.h"
+#include "ISettingsSection.h"
 #include "Layout/Visibility.h"
+#include "Misc/NotifyHook.h"
 #include "Widgets/DeclarativeSyntaxSupport.h"
 #include "Widgets/SWidget.h"
 #include "Widgets/SCompoundWidget.h"
 #include "Widgets/SOverlay.h"
-#include "ISettingsCategory.h"
-#include "ISettingsContainer.h"
-#include "ISettingsEditorModel.h"
-#include "ISettingsSection.h"
 
 class IDetailsView;
+class SAssetSearchBox;
+class SScrollBox;
 class SVerticalBox;
+struct FAssetSearchBoxSuggestion;
 struct FPropertyChangedEvent;
 
 /**
@@ -34,7 +38,7 @@ public:
 public:
 
 	/** Destructor. */
-	~SSettingsEditor();
+	virtual ~SSettingsEditor() override;
 
 public:
 
@@ -123,13 +127,23 @@ private:
 
 	/** Callback for importing all sections from a single ini. */
 	FReply HandleImportButtonClicked();
-private:
+
+	FText GetSearchToolTipText() const;
+	void OnSearchTextChanged(const FText& InText);
+	void OnSearchTextCommitted(const FText& InText, ETextCommit::Type InType);
+	void OnAssetSearchSuggestionFilter(const FText& InText, TArray<FAssetSearchBoxSuggestion>& OutSuggestion, FText& OutText);
+	FText OnAssetSearchSuggestionChosen(const FText& InText, const FString& InSuggestion);
+	void HandleSearchTermsChanged(const ISettingsEditorModule::FSearchTerm& InNewSearchTerm);
+	
+	ECheckBoxState IsSearchFilterButtonActive(FString InFilter) const;
+	void OnSearchFilterButtonClicked(ECheckBoxState InState, FString InFilter);
 
 	/** Holds the vertical box for settings categories. */
 	TSharedPtr<SVerticalBox> CategoriesBox;
 
 	/** Holds the overlay slot for custom widgets. */
 	SOverlay::FOverlaySlot* CustomWidgetSlot;
+
 	/** Holds a pointer to the view model. */
 	ISettingsEditorModelPtr Model;
 
@@ -147,4 +161,13 @@ private:
 	
 	/** Are we showing all settings at once */
 	bool bShowingAllSettings;
+
+	/** Actual search box with suggestions */
+	TSharedPtr<SAssetSearchBox> SearchBox;
+
+	/** Used to store quick filters checkboxes representing the keys of the suggestions (cvar, keybind, ...) */
+	TSharedPtr<SScrollBox> SearchFiltersBox;
+	
+	/** Used to keep track of the search filters added */
+	TSet<FString> QuickSearchFiltersAdded;
 };
diff --git a/Engine/Source/Developer/SettingsEditor/Public/ISettingsEditorModule.h b/Engine/Source/Developer/SettingsEditor/Public/ISettingsEditorModule.h
--- a/Engine/Source/Developer/SettingsEditor/Public/ISettingsEditorModule.h
+++ b/Engine/Source/Developer/SettingsEditor/Public/ISettingsEditorModule.h
@@ -65,8 +65,38 @@ public:
 	 */
 	virtual void UpdateSettings(bool bForce = false) = 0;
 
-public:
+	struct FSearchTerm
+	{
+		/** In what context is this search term used (i.e. project/editor) */
+		FName ContainerName;
+		
+		/** The key used to show suggestions values (i.e. cvar/keybind) */
+		FString Key;
+		
+		/** The value used for the actual search along the key */
+		FString Value;
+	};
+
+	/**
+	 * Registers search terms for suggestions
+	 * 
+	 * @param InSearchTerm The context, key and value to register
+	 */
+	virtual void RegisterSearchTerm(const FSearchTerm& InSearchTerm) = 0;
+
+	/** 
+	 * Delegate called when new search terms have been registered
+	 */
+	DECLARE_MULTICAST_DELEGATE_OneParam(FOnSearchTermsChanged, const FSearchTerm& /** SearchTerm */)
+	virtual FOnSearchTermsChanged::RegistrationType& OnSearchTermsChanged() = 0;
+	
+	/**
+	 * Find search terms based on container name
+	 * 
+	 * @param InContainerName The container name for context
+	 */
+	virtual const TMap<FString, TSet<FString>>* FindSearchTerms(FName InContainerName) const = 0;
 
 	/** Virtual destructor. */
-	virtual ~ISettingsEditorModule() = default;
+	virtual ~ISettingsEditorModule() override = default;
 };
diff --git a/Engine/Source/Developer/SettingsEditor/SettingsEditor.Build.cs b/Engine/Source/Developer/SettingsEditor/SettingsEditor.Build.cs
--- a/Engine/Source/Developer/SettingsEditor/SettingsEditor.Build.cs
+++ b/Engine/Source/Developer/SettingsEditor/SettingsEditor.Build.cs
@@ -24,6 +24,7 @@ public class SettingsEditor : ModuleRules
                 "CoreUObject",
 				"DesktopPlatform",
 				"DeveloperSettings",
+				"EditorWidgets",
 				"Engine",
                 "InputCore",
 				"PropertyEditor",
@@ -31,7 +32,8 @@ public class SettingsEditor : ModuleRules
 				"Slate",
 				"SlateCore",
 				"SourceControl",
-				"ToolMenus"
+				"ToolMenus",
+				"ToolWidgets",
             }
         );
 	}
diff --git a/Engine/Source/Editor/InputBindingEditor/InputBindingEditor.Build.cs b/Engine/Source/Editor/InputBindingEditor/InputBindingEditor.Build.cs
--- a/Engine/Source/Editor/InputBindingEditor/InputBindingEditor.Build.cs
+++ b/Engine/Source/Editor/InputBindingEditor/InputBindingEditor.Build.cs
@@ -7,12 +7,12 @@ public class InputBindingEditor : ModuleRules
 {
 	public InputBindingEditor(ReadOnlyTargetRules Target) : base(Target)
 	{
-
 		PrivateDependencyModuleNames.AddRange(
 			new string[]
 			{
 				"Core",
 				"CoreUObject",
+				"DeveloperSettings",
 				"Slate",
 				"SlateCore",
 				"InputCore",
@@ -21,13 +21,6 @@ public class InputBindingEditor : ModuleRules
 				"UnrealEd",
 				"PropertyEditor",
 				"Settings",
-				"DeveloperSettings"
-			}
-		);
-
-		DynamicallyLoadedModuleNames.AddRange(
-			new string[]
-			{
 				"SettingsEditor",
 			}
 		);
diff --git a/Engine/Source/Editor/InputBindingEditor/Private/InputBindingEditorModule.cpp b/Engine/Source/Editor/InputBindingEditor/Private/InputBindingEditorModule.cpp
--- a/Engine/Source/Editor/InputBindingEditor/Private/InputBindingEditorModule.cpp
+++ b/Engine/Source/Editor/InputBindingEditor/Private/InputBindingEditorModule.cpp
@@ -21,6 +21,9 @@
 #include "HAL/PlatformMisc.h"
 #include "IDetailCustomization.h"
 #include "IDetailGroup.h"
+#include "ISettingsEditorModule.h"
+#include "ISettingsEditorModule.h"
+#include "ISettingsEditorModule.h"
 #include "ISettingsModule.h"
 #include "ISettingsSection.h"
 #include "Interfaces/IInputBindingEditorModule.h"
@@ -224,6 +227,14 @@ public:
 	void UpdateUI()
 	{
 		FInputBindingManager& InputBindingManager = FInputBindingManager::Get();
+		
+		ISettingsEditorModule* SettingsEditorModule = FModuleManager::GetModulePtr<ISettingsEditorModule>(TEXT("SettingsEditor"));
+		const FString SearchKey = TEXT("keybind");
+		FName SearchContext = NAME_None;
+		if (TSharedPtr<IDetailsView> DetailsView = DetailBuilder->GetDetailsViewSharedPtr())
+		{
+			SearchContext = DetailsView->GetIdentifier();
+		}
 
 		for (TSharedPtr<FChordTreeItem>& TreeItem : ContextList)
 		{
@@ -276,7 +287,7 @@ public:
 					Row = &CategoryBuilder.AddCustomRow(CommandInfo->GetLabel());
 				}
 
-				// Set up search filter (for i.e. KeyBinding="F")
+				// Set up search filter (for i.e. keybind=F)
 				const TSharedRef<const FInputChord> FirstInputChord = CommandInfo->GetActiveChord(EMultipleKeyBindingIndex::Primary);
 				const TSharedRef<const FInputChord> SecondInputChord = CommandInfo->GetActiveChord(EMultipleKeyBindingIndex::Secondary);
 
@@ -285,18 +296,44 @@ public:
 
 				if (bFirstChordValid || bSecondChordValid)
 				{
-					static const FTextFormat KeyFormat = LOCTEXT("SearchKeyFilter", "KeyBinding=\"{0}\"");
+					static const FTextFormat KeyFormat = LOCTEXT("SearchKeyFilter", "keybind={0}");
 					static const FTextFormat TokenCombineFormat = INVTEXT("{0} {1}");
 
 					FText FirstKeyFilter, SecondKeyFilter;
 					if (bFirstChordValid)
 					{
-						FirstKeyFilter = FText::FormatOrdered(KeyFormat, FirstInputChord->GetInputText());
+						FText InputText = FirstInputChord->GetInputText();
+						FirstKeyFilter = FText::FormatOrdered(KeyFormat, InputText);
+
+						// Register search key and value for settings editor suggestions
+						if (SettingsEditorModule)
+						{
+							ISettingsEditorModule::FSearchTerm SearchTerm
+							{
+								.ContainerName = SearchContext,
+								.Key = SearchKey,
+								.Value = InputText.ToString()
+							};
+							SettingsEditorModule->RegisterSearchTerm(SearchTerm);
+						}
 					}
 
 					if (bSecondChordValid)
 					{
-						SecondKeyFilter = FText::FormatOrdered(KeyFormat, FirstInputChord->GetInputText());
+						FText InputText = SecondInputChord->GetInputText();
+						SecondKeyFilter = FText::FormatOrdered(KeyFormat, InputText);
+
+						// Register search key and value for settings editor suggestions
+						if (SettingsEditorModule)
+						{
+							ISettingsEditorModule::FSearchTerm SearchTerm
+							{
+								.ContainerName = SearchContext,
+								.Key = SearchKey,
+								.Value = InputText.ToString()
+							};
+							SettingsEditorModule->RegisterSearchTerm(SearchTerm);
+						}
 					}
 
 					FText FilterText = bFirstChordValid && bSecondChordValid
@@ -316,33 +353,19 @@ public:
 				.MaxDesiredWidth(0)
 				.MinDesiredWidth(500)
 				[
-					SNew(SVerticalBox)
-					+ SVerticalBox::Slot()
-					.AutoHeight()
-					[
-						SNew(STextBlock)
-						.Text(CommandInfo->GetLabel())
-						.ToolTipText_Lambda([CommandInfo]() -> FText
-						{
-							FText CommandInfoTooltip = CommandInfo->GetDescription();
+					SNew(STextBlock)
+					.Text(CommandInfo->GetLabel())
+					.ToolTipText_Lambda([CommandInfo]() -> FText
+					{
+						FText CommandInfoTooltip = CommandInfo->GetDescription();
 
-							if (InputBindingEditorModule::bShowBindingNames)
-							{
-								CommandInfoTooltip = FText::Format(LOCTEXT("CommandInfoDebugToolTip", "{0}\n\nBinding Context: {1}\nCommand Name: {2}"), CommandInfoTooltip, FText::FromName(CommandInfo->GetBindingContext()), FText::FromName(CommandInfo->GetCommandName()));
-							}
+						if (InputBindingEditorModule::bShowBindingNames)
+						{
+							CommandInfoTooltip = FText::Format(LOCTEXT("CommandInfoDebugToolTip", "{0}\n\nBinding Context: {1}\nCommand Name: {2}"), CommandInfoTooltip, FText::FromName(CommandInfo->GetBindingContext()), FText::FromName(CommandInfo->GetCommandName()));
+						}
 
-							return CommandInfoTooltip;
-						})
-					]
-					+ SVerticalBox::Slot()
-					.Padding(0.0f, 3.0f, 0.0f, 3.0f)
-					.AutoHeight()
-					[
-						SNew(STextBlock)
-						.Font(IDetailLayoutBuilder::GetDetailFont())
-						.ColorAndOpacity(FLinearColor::Gray)
-						.Text(CommandInfo->GetDescription())
-					]
+						return CommandInfoTooltip;
+					})
 				];
 
 				Row->ValueContent()
diff --git a/Engine/Source/Editor/PropertyEditor/Private/SDetailsView.cpp b/Engine/Source/Editor/PropertyEditor/Private/SDetailsView.cpp
--- a/Engine/Source/Editor/PropertyEditor/Private/SDetailsView.cpp
+++ b/Engine/Source/Editor/PropertyEditor/Private/SDetailsView.cpp
@@ -674,13 +674,6 @@ void SDetailsView::SetRootObjectCustomizationInstance(TSharedPtr<IDetailRootObje
 	RerunCurrentFilter();
 }
 
-void SDetailsView::ClearSearch()
-{
-	CurrentFilter.FilterStrings.Empty();
-	SearchBox->SetText(FText::GetEmpty());
-	RerunCurrentFilter();
-}
-
 void SDetailsView::SetObjectFilter(TSharedPtr<FDetailsViewObjectFilter> InFilter)
 {
 	ObjectFilter = InFilter;
diff --git a/Engine/Source/Editor/PropertyEditor/Private/SDetailsView.h b/Engine/Source/Editor/PropertyEditor/Private/SDetailsView.h
--- a/Engine/Source/Editor/PropertyEditor/Private/SDetailsView.h
+++ b/Engine/Source/Editor/PropertyEditor/Private/SDetailsView.h
@@ -52,7 +52,6 @@ public:
 	virtual void RemoveInvalidObjects() override;
 	virtual void SetObjectPackageOverrides(const TMap<TWeakObjectPtr<UObject>, TWeakObjectPtr<UPackage>>& InMapping) override;
 	virtual void SetRootObjectCustomizationInstance(TSharedPtr<IDetailRootObjectCustomization> InRootObjectCustomization) override;
-	virtual void ClearSearch() override;
 	virtual void SetObjectFilter(TSharedPtr<FDetailsViewObjectFilter> InFilter) override;
 	virtual void SetClassViewerFilters(const TArray<TSharedRef<class IClassViewerFilter>>& InFilters) override;
 
diff --git a/Engine/Source/Editor/PropertyEditor/Private/SDetailsViewBase.cpp b/Engine/Source/Editor/PropertyEditor/Private/SDetailsViewBase.cpp
--- a/Engine/Source/Editor/PropertyEditor/Private/SDetailsViewBase.cpp
+++ b/Engine/Source/Editor/PropertyEditor/Private/SDetailsViewBase.cpp
@@ -64,6 +64,18 @@ SDetailsViewBase::~SDetailsViewBase()
 	}
 }
 
+void SDetailsViewBase::ClearSearch()
+{
+	CurrentFilter.FilterStrings.Empty();
+	SearchBox->SetText(FText::GetEmpty());
+	RerunCurrentFilter();
+}
+
+void SDetailsViewBase::SetSearchText(const FText& InSearch)
+{
+	SearchBox->SetText(InSearch);
+}
+
 void SDetailsViewBase::OnGetChildrenForDetailTree(TSharedRef<FDetailTreeNode> InTreeNode, TArray< TSharedRef<FDetailTreeNode> >& OutChildren)
 {
 	InTreeNode->GetChildren(OutChildren);
diff --git a/Engine/Source/Editor/PropertyEditor/Private/SDetailsViewBase.h b/Engine/Source/Editor/PropertyEditor/Private/SDetailsViewBase.h
--- a/Engine/Source/Editor/PropertyEditor/Private/SDetailsViewBase.h
+++ b/Engine/Source/Editor/PropertyEditor/Private/SDetailsViewBase.h
@@ -63,6 +63,9 @@ public:
 	{
 		return CurrentFilter.FilterStrings.Num() > 0;
 	}
+	
+	virtual void ClearSearch() override;
+	virtual void SetSearchText(const FText& InSearch) override;
 
 	virtual int32 GetNumVisibleTopLevelObjects() const override
 	{
diff --git a/Engine/Source/Editor/PropertyEditor/Private/SStructureDetailsView.cpp b/Engine/Source/Editor/PropertyEditor/Private/SStructureDetailsView.cpp
--- a/Engine/Source/Editor/PropertyEditor/Private/SStructureDetailsView.cpp
+++ b/Engine/Source/Editor/PropertyEditor/Private/SStructureDetailsView.cpp
@@ -301,13 +301,6 @@ void SStructureDetailsView::InvalidateCachedState()
 	}
 }
 
-void SStructureDetailsView::ClearSearch()
-{
-	CurrentFilter.FilterStrings.Empty();
-	SearchBox->SetText(FText::GetEmpty());
-	RerunCurrentFilter();
-}
-
 const TArray< TWeakObjectPtr<UObject> >& SStructureDetailsView::GetSelectedObjects() const
 {
 	static const TArray< TWeakObjectPtr<UObject> > DummyRef;
diff --git a/Engine/Source/Editor/PropertyEditor/Private/SStructureDetailsView.h b/Engine/Source/Editor/PropertyEditor/Private/SStructureDetailsView.h
--- a/Engine/Source/Editor/PropertyEditor/Private/SStructureDetailsView.h
+++ b/Engine/Source/Editor/PropertyEditor/Private/SStructureDetailsView.h
@@ -81,8 +81,6 @@ public:
 	/** Invalidates cached state such as the "revert to default" arrow and edit conditions, without rebuilding the entire panel. */
 	virtual void InvalidateCachedState() override;
 	virtual void MoveScrollOffset(int32 DeltaOffset) override {}
-	virtual void ClearSearch() override;
-public:
 
 	// IDetailsView interface
 	virtual const TArray< TWeakObjectPtr<UObject> >& GetSelectedObjects() const override;
diff --git a/Engine/Source/Editor/PropertyEditor/Public/IDetailsView.h b/Engine/Source/Editor/PropertyEditor/Public/IDetailsView.h
--- a/Engine/Source/Editor/PropertyEditor/Public/IDetailsView.h
+++ b/Engine/Source/Editor/PropertyEditor/Public/IDetailsView.h
@@ -118,6 +118,11 @@ public:
 	 * Clears any search terms in the current filter
 	 */
 	virtual void ClearSearch() = 0;
+	
+	/** 
+	 * Sets the active search for the current filter
+	 */
+	virtual void SetSearchText(const FText& InSearch) = 0;
 
 	/**
 	 * @return The number of visible top level objects. This value is affected by filtering. 
