EditorUsability : SettingsEditor
- Improved settings editor search following feedback (part 2)
- Replaced custom widget toolbar by one based on UToolMenu (with a custom style to have transparent background)
- Added checkbox menu entries with combo box that shows all available values for a filter key
- Added option to search for unbound keybinds (eg: keybind=unbound)
- Added "context" search key and values to filter by input binding context name
- Search keys (keybind, cvar, context) can now be localized
- Reduced search box suggestion list height to not hide filtered results (requires part 1 changes)
- Swapped "Export" and "Import" buttons order

#Synced-CL 49549735

diff --git a/Engine/Source/Developer/SettingsEditor/Private/Customizations/SettingsEditorCustomization.cpp b/Engine/Source/Developer/SettingsEditor/Private/Customizations/SettingsEditorCustomization.cpp
--- a/Engine/Source/Developer/SettingsEditor/Private/Customizations/SettingsEditorCustomization.cpp
+++ b/Engine/Source/Developer/SettingsEditor/Private/Customizations/SettingsEditorCustomization.cpp
@@ -10,6 +10,8 @@
 #include "Modules/ModuleManager.h"
 #include "UObject/Package.h"
 
+#define LOCTEXT_NAMESPACE "SettingsEditorCustomization"
+
 TSharedRef<IDetailCustomization> FSettingsEditorCustomization::MakeInstance()
 {
 	return MakeShared<FSettingsEditorCustomization>();
@@ -75,13 +77,15 @@ void FSettingsEditorCustomization::CustomizeDetails(IDetailLayoutBuilder& InDeta
 	InDetailLayoutBuilder.GetObjectsBeingCustomized(SettingsObjects);
 	
 	ISettingsEditorModule* SettingsEditorModule = FModuleManager::GetModulePtr<ISettingsEditorModule>(TEXT("SettingsEditor"));
-	const FString SearchKey = TEXT("cvar");
 	FName SearchContext = NAME_None;
 	if (const TSharedPtr<IDetailsView> DetailsView = InDetailLayoutBuilder.GetDetailsViewSharedPtr())
 	{
 		SearchContext = DetailsView->GetIdentifier();
 	}
 
+	static const FText SearchKeyCVar = LOCTEXT("SearchKeyCVar", "CVar");
+	static const FTextFormat SearchFormat = INVTEXT("{0}={1}");
+	static const FTextFormat TokenDelimiter = INVTEXT("\"{0}\"");
 	for (const TWeakObjectPtr<UObject>& SettingsObject : SettingsObjects)
 	{
 		TStrongObjectPtr<UObject> PinnedSettingsObject = SettingsObject.Pin();
@@ -124,26 +128,29 @@ void FSettingsEditorCustomization::CustomizeDetails(IDetailLayoutBuilder& InDeta
 					PropertyHandle->CreatePropertyValueWidget()
 				];
 
-				// Allow searching of cvar in settings editor (for i.e. cvar=r.vsync)
-				Row.FilterString(FText::FromString(TEXT("cvar=") + CVar));
-
 				// Override reset to default
-				Row.OverrideResetToDefault(FResetToDefaultOverride::Create(
-					FIsResetToDefaultVisible::CreateSP(this, &FSettingsEditorCustomization::IsResetToDefaultVisible),
-					FResetToDefaultHandler::CreateSP(this, &FSettingsEditorCustomization::OnResetToDefault))
-				);
-				
+                Row.OverrideResetToDefault(FResetToDefaultOverride::Create(
+                	FIsResetToDefaultVisible::CreateSP(this, &FSettingsEditorCustomization::IsResetToDefaultVisible),
+                	FResetToDefaultHandler::CreateSP(this, &FSettingsEditorCustomization::OnResetToDefault))
+                );
+
+				// Allow searching of cvar in settings editor (for i.e. cvar="r.vsync")
+				FText FilterText = FText::FormatOrdered(TokenDelimiter, FText::FromString(CVar));
+
 				// Register search key and value for settings editor suggestions
 				if (SettingsEditorModule)
 				{
 					ISettingsEditorModule::FSearchTerm SearchTerm
 					{
 						.ContainerName = SearchContext,
-						.Key = SearchKey,
-						.Value = CVar
+						.Key = SearchKeyCVar.ToString(),
+						.Value = FilterText.ToString()
 					};
 					SettingsEditorModule->RegisterSearchTerm(SearchTerm);
 				}
+
+				FilterText = FText::FormatOrdered(SearchFormat, SearchKeyCVar.ToLower(), FilterText);
+				Row.FilterString(FilterText);
 			}
 			else
 			{
@@ -326,3 +333,5 @@ void FSettingsEditorCustomization::ResetDefaultPropertyValues(IDetailLayoutBuild
 		DefaultPropertyValues.Add(PropertyHandle->GetProperty()->GetFName(), DefaultValue);
 	}
 }
+
+#undef LOCTEXT_NAMESPACE
diff --git a/Engine/Source/Developer/SettingsEditor/Private/Models/SettingsEditorMenuContext.h b/Engine/Source/Developer/SettingsEditor/Private/Models/SettingsEditorMenuContext.h
new file mode 100644
--- /dev/null
+++ b/Engine/Source/Developer/SettingsEditor/Private/Models/SettingsEditorMenuContext.h
@@ -0,0 +1,24 @@
+// Copyright Epic Games, Inc. All Rights Reserved.
+
+#pragma once
+
+#include "Templates/SharedPointer.h"
+#include "UObject/Object.h"
+#include "SettingsEditorMenuContext.generated.h"
+
+class SSettingsEditor;
+
+/** Context passed in UToolMenu when generating menu entries for settings editor search toolbar */
+UCLASS()
+class USettingsEditorMenuContext : public UObject
+{
+	GENERATED_BODY()
+
+public:
+	TSharedPtr<SSettingsEditor> GetSettingsEditor() const
+	{
+		return SettingsEditorWeak.Pin();
+	}
+
+	TWeakPtr<SSettingsEditor> SettingsEditorWeak;
+};
\ No newline at end of file
diff --git a/Engine/Source/Developer/SettingsEditor/Private/SettingsEditorModule.cpp b/Engine/Source/Developer/SettingsEditor/Private/SettingsEditorModule.cpp
--- a/Engine/Source/Developer/SettingsEditor/Private/SettingsEditorModule.cpp
+++ b/Engine/Source/Developer/SettingsEditor/Private/SettingsEditorModule.cpp
@@ -204,12 +204,16 @@ public:
 			return;
 		}
 
+		bool bAlreadyRegistered = false;
 		SearchTerms
 			.FindOrAdd(InSearchTerm.ContainerName)
-			.FindOrAdd(InSearchTerm.Key.ToLower())
-			.Add(InSearchTerm.Value);
+			.FindOrAdd(InSearchTerm.Key)
+			.Add(InSearchTerm.Value, &bAlreadyRegistered);
 
-		OnSearchTermsChangedDelegate.Broadcast(InSearchTerm);
+		if (!bAlreadyRegistered)
+		{
+			OnSearchTermsChangedDelegate.Broadcast(InSearchTerm);
+		}
 	}
 	
 	virtual FOnSearchTermsChanged::RegistrationType& OnSearchTermsChanged() override
diff --git a/Engine/Source/Developer/SettingsEditor/Private/Widgets/SSettingsEditor.cpp b/Engine/Source/Developer/SettingsEditor/Private/Widgets/SSettingsEditor.cpp
--- a/Engine/Source/Developer/SettingsEditor/Private/Widgets/SSettingsEditor.cpp
+++ b/Engine/Source/Developer/SettingsEditor/Private/Widgets/SSettingsEditor.cpp
@@ -24,6 +24,9 @@
 #include "Widgets/Input/SHyperlink.h"
 #include "Widgets/SSettingsSectionHeader.h"
 #include "SSettingsEditorCheckoutNotice.h"
+#include "ToolMenuEntry.h"
+#include "ToolMenus.h"
+#include "ToolMenu.h"
 #include "HAL/PlatformFileManager.h"
 #include "HAL/PlatformFile.h"
 #include "Containers/Ticker.h"
@@ -32,6 +35,7 @@
 #include "HAL/FileManager.h"
 #include "Misc/ConfigCacheIni.h"
 #include "Misc/MessageDialog.h"
+#include "Models/SettingsEditorMenuContext.h"
 
 #define LOCTEXT_NAMESPACE "SSettingsEditor"
 
@@ -45,35 +49,6 @@ SSettingsEditor::~SSettingsEditor()
 	FInternationalization::Get().OnCultureChanged().RemoveAll(this);
 }
 
-void SSettingsEditor::HandleSearchTermsChanged(const ISettingsEditorModule::FSearchTerm& InNewSearchTerm)
-{
-	if (InNewSearchTerm.ContainerName != SettingsContainer->GetName() || QuickSearchFiltersAdded.Contains(InNewSearchTerm.Key))
-	{
-		return;
-	}
-
-	QuickSearchFiltersAdded.Add(InNewSearchTerm.Key);
-
-	const FString SearchKey = InNewSearchTerm.Key;
-	SearchFiltersBox->AddSlot()
-	.Padding(0.f)
-    [
-    	SNew(SBox)
-    	.Padding(0.f, 8.f, 8.f, 8.f)
-    	[
-    		SNew(SCheckBox)
-    		.Style(FAppStyle::Get(), "DetailsView.SectionButton")
-    		.OnCheckStateChanged(this, &SSettingsEditor::OnSearchFilterButtonClicked, SearchKey)
-    		.IsChecked(this, &SSettingsEditor::IsSearchFilterButtonActive, SearchKey)
-    		[
-    			SNew(STextBlock)
-				.Text(FText::FromString(SearchKey))
-    			.ToolTipText(LOCTEXT("SearchFilterToolTip", "Toggles this search filter on/off"))
-    		]
-    	]
-    ];
-}
-
 
 /* SSettingsEditor interface
  *****************************************************************************/
@@ -119,9 +94,28 @@ void SSettingsEditor::Construct( const FArguments& InArgs, const ISettingsEditor
 	TSharedPtr<FSettingsDetailRootObjectCustomization> RootObjectCustomization = MakeShareable(new FSettingsDetailRootObjectCustomization(Model, SettingsView.ToSharedRef()));
 	RootObjectCustomization->Initialize();
 	SettingsView->SetRootObjectCustomizationInstance(RootObjectCustomization);
+	
+	const FName SettingsEditorToolbarName = TEXT("SettingsEditor.SearchToolbar");
+	UToolMenus* ToolMenus = UToolMenus::Get();
+	if (!ToolMenus->IsMenuRegistered(SettingsEditorToolbarName))
+	{
+		UToolMenu* const SettingsToolbarMenu = ToolMenus->RegisterMenu(
+			SettingsEditorToolbarName, NAME_None /* parent */, EMultiBoxType::SlimHorizontalToolBar
+		);
+
+		SettingsToolbarMenu->StyleName = "NoBackgroundViewportToolbar";
+		SettingsToolbarMenu->bSeparateSections = false;
+		SettingsToolbarMenu->AddDynamicSection(
+			TEXT("Filters"), 
+			FNewToolMenuDelegate::CreateStatic(&SSettingsEditor::FillSettingsFiltersSection)
+		);
+	}
 
-	TSharedPtr<SScrollBox> ScrollBox;
+	USettingsEditorMenuContext* MenuContext = NewObject<USettingsEditorMenuContext>();
+	MenuContext->SettingsEditorWeak = SharedThis(this);
+	const TSharedRef<SWidget> MenuWidget = ToolMenus->GenerateWidget(SettingsEditorToolbarName, FToolMenuContext(MenuContext));
 
+	TSharedPtr<SScrollBox> ScrollBox;
 	ChildSlot
 	[
 		SNew(SVerticalBox)
@@ -177,6 +171,7 @@ void SSettingsEditor::Construct( const FArguments& InArgs, const ISettingsEditor
 					.Padding(0.f, 0.f, 0.f, 0.f)
 					[
 						SAssignNew(SearchBox, SAssetSearchBox)
+						.SuggestionListHeight(50.f)
 						.HintText(LOCTEXT("FilterSearch", "Search..."))
 						.ShowSearchHistory(true)
 						.ToolTipText(this, &SSettingsEditor::GetSearchToolTipText)
@@ -192,30 +187,29 @@ void SSettingsEditor::Construct( const FArguments& InArgs, const ISettingsEditor
 					.VAlign(VAlign_Center)
 					.Padding(8.0f, 0.0f, 0.0f, 0.0f)
 					[
-						// Export button
+						// Import button
 						SNew(SButton)
-						.OnClicked(this, &SSettingsEditor::HandleExportButtonClicked)
-						.Text(LOCTEXT("ExportButtonText", "Export..."))
-						.ToolTipText(LOCTEXT("ExportAllButtonTooltip", "Export all editor settings to a file on your computer"))
+						.OnClicked(this, &SSettingsEditor::HandleImportButtonClicked)
+						.Text(LOCTEXT("ImportButtonText", "Import..."))
+						.ToolTipText(LOCTEXT("ImportAllButtonTooltip", "Import all editor settings from a file on your computer"))
 					]
 					+ SHorizontalBox::Slot()
 					.AutoWidth()
 					.VAlign(VAlign_Center)
 					.Padding(8.0f, 0.0f, 6.0f, 0.0f)
 					[
-						// Import button
+						// Export button
 						SNew(SButton)
-						.OnClicked(this, &SSettingsEditor::HandleImportButtonClicked)
-						.Text(LOCTEXT("ImportButtonText", "Import..."))
-						.ToolTipText(LOCTEXT("ImportAllButtonTooltip", "Import all editor settings from a file on your computer"))
+						.OnClicked(this, &SSettingsEditor::HandleExportButtonClicked)
+						.Text(LOCTEXT("ExportButtonText", "Export..."))
+						.ToolTipText(LOCTEXT("ExportAllButtonTooltip", "Export all editor settings to a file on your computer"))
 					]
 				]
 				+ SVerticalBox::Slot()
+				.Padding(0.f, 4.f)
 				.AutoHeight()
 				[
-					SAssignNew(SearchFiltersBox, SScrollBox)
-                    .Orientation(EOrientation::Orient_Horizontal)
-                    .ScrollBarThickness(FVector2D(1.f))
+					MenuWidget
 				]
 				+ SVerticalBox::Slot()
 				.FillHeight(1.0f)
@@ -238,12 +232,6 @@ void SSettingsEditor::Construct( const FArguments& InArgs, const ISettingsEditor
 	FInternationalization::Get().OnCultureChanged().AddSP(this, &SSettingsEditor::HandleCultureChanged);
 	Model->OnSelectionChanged().AddSP(this, &SSettingsEditor::HandleModelSelectionChanged);
 	SettingsContainer->OnCategoryModified().AddSP(this, &SSettingsEditor::HandleSettingsContainerCategoryModified);
-	
-	// Update quick filters keys based on new search terms registered
-	if (ISettingsEditorModule* SettingsEditorModule = FModuleManager::GetModulePtr<ISettingsEditorModule>(TEXT("SettingsEditor")))
-	{
-		SettingsEditorModule->OnSearchTermsChanged().AddSP(this, &SSettingsEditor::HandleSearchTermsChanged);
-	}
 
 	ReloadCategories();
 }
@@ -863,14 +851,16 @@ void SSettingsEditor::OnAssetSearchSuggestionFilter(const FText& InText, TArray<
 		return;
 	}
 
+	static const FText SettingsCategory = LOCTEXT("SettingsEditor", "Settings");
+	static const FTextFormat SearchFormat = INVTEXT("{0}={1}");
 	const FString ActiveSearch = InText.ToString();
-	const FText SettingsCategory = LOCTEXT("SettingsEditor", "Settings");
 	for (const TPair<FString, TSet<FString>>& SearchTerm : *SearchTerms)
 	{
 		if (ActiveSearch.Contains(SearchTerm.Key) || SearchTerm.Key.Contains(ActiveSearch))
 		{
-			FString SearchSuggestion = SearchTerm.Key + TEXT("=");
-			OutSuggestion.Add(FAssetSearchBoxSuggestion(SearchSuggestion, FText::FromString(SearchSuggestion), SettingsCategory));
+			FString SearchSuggestion = SearchTerm.Key.ToLower() + TEXT("=");
+			FText SearchSuggestionLabel = FText::FormatOrdered(SearchFormat, FText::FromString(SearchTerm.Key), FText::GetEmpty());
+			OutSuggestion.Add(FAssetSearchBoxSuggestion(SearchSuggestion, SearchSuggestionLabel, SettingsCategory));
 			FString SearchCopy = ActiveSearch;
 			SearchCopy.RemoveFromStart(SearchSuggestion);
 
@@ -878,8 +868,9 @@ void SSettingsEditor::OnAssetSearchSuggestionFilter(const FText& InText, TArray<
 			{
 				if (SearchCopy.IsEmpty() || SearchValue.StartsWith(SearchCopy))
 				{
-					SearchSuggestion = SearchTerm.Key + TEXT("=") + SearchValue;
-					OutSuggestion.Add(FAssetSearchBoxSuggestion(SearchSuggestion, FText::FromString(SearchSuggestion), SettingsCategory));
+					SearchSuggestion = SearchTerm.Key.ToLower() + TEXT("=") + SearchValue;
+					SearchSuggestionLabel = FText::FormatOrdered(SearchFormat, FText::FromString(SearchTerm.Key), FText::FromString(SearchValue));
+					OutSuggestion.Add(FAssetSearchBoxSuggestion(SearchSuggestion, SearchSuggestionLabel, SettingsCategory));
 				}
 			}
 		}
@@ -913,15 +904,91 @@ FText SSettingsEditor::OnAssetSearchSuggestionChosen(const FText& InText, const
 	return FText::FromString(CurrentSearch);
 }
 
-ECheckBoxState SSettingsEditor::IsSearchFilterButtonActive(FString InFilter) const
+void SSettingsEditor::FillSettingsFiltersSection(UToolMenu* InMenu)
 {
-	const FString ActiveSearch = SearchBox->GetText().ToString();
-	return ActiveSearch.StartsWith(InFilter) ? ECheckBoxState::Checked : ECheckBoxState::Unchecked;
+	if (!InMenu)
+	{
+		return;
+	}
+
+	const USettingsEditorMenuContext* Context = InMenu->FindContext<USettingsEditorMenuContext>();
+	if (!Context)
+	{
+		return;
+	}
+
+	const TSharedPtr<SSettingsEditor> SettingsEditor = Context->GetSettingsEditor();
+	if (!SettingsEditor)
+	{
+		return;
+	}
+
+	const ISettingsEditorModule* SettingsEditorModule = FModuleManager::GetModulePtr<ISettingsEditorModule>(TEXT("SettingsEditor"));
+	if (!SettingsEditorModule || !SettingsEditor->SettingsContainer)
+	{
+		return;
+	}
+
+	const TMap<FString, TSet<FString>>* SearchTerms = SettingsEditorModule->FindSearchTerms(SettingsEditor->SettingsContainer->GetName());
+	if (!SearchTerms)
+	{
+		return;
+	}
+
+	FToolMenuSection& FiltersSection = InMenu->FindOrAddSection("Filters", LOCTEXT("FiltersLabel", "Filters"));
+	for (const TPair<FString, TSet<FString>>& SearchTerm : *SearchTerms)
+	{
+		FiltersSection
+			.AddEntry(CreateFiltersCheckboxEntry(Context, SearchTerm.Key, SearchTerm.Value))
+			.SetShowInToolbarTopLevel(true);
+	}
+}
+
+FToolMenuEntry SSettingsEditor::CreateFiltersCheckboxEntry(const USettingsEditorMenuContext* InContext, const FString& InKey, const TSet<FString>& InValues)
+{
+	const TSharedPtr<SSettingsEditor> SettingsEditor = InContext->GetSettingsEditor();
+	checkf(SettingsEditor.IsValid(), TEXT("SettingsEditor must be valid"))
+
+	FUIAction FilterAction;
+	FilterAction.ExecuteAction = FExecuteAction::CreateSP(SettingsEditor.Get(), &SSettingsEditor::ToggleSearchFilterAction, InKey);
+	FilterAction.GetActionCheckState = FGetActionCheckState::CreateSP(SettingsEditor.Get(), &SSettingsEditor::IsSearchFilterActive, InKey, /** Equals */false);
+
+	FToolMenuEntry FilterEntry = FToolMenuEntry::InitSubMenu(
+		FName(InKey),
+		FText::FromString(InKey),
+		LOCTEXT("SearchFilterToolTip", "Toggles this search filter on/off"),
+		FNewToolMenuDelegate::CreateLambda([SettingsEditor, InKey, InValues](UToolMenu* InSubmenu)
+		{
+			FToolMenuSection& FiltersSubSection = InSubmenu->FindOrAddSection(FName(InKey), FText::FromString(InKey));
+			for (const FString& Value : InValues)
+			{
+				const FString SearchTerm = InKey + TEXT("=") + Value;
+				FiltersSubSection.AddMenuEntry(
+					NAME_None,
+					FText::FromString(Value.TrimChar('"')),
+					LOCTEXT("SetFilterToolTip", "Apply search filter with this value"),
+					FSlateIcon(),
+					FUIAction(
+						FExecuteAction::CreateSP(SettingsEditor.Get(), &SSettingsEditor::SetSearchFilter, SearchTerm),
+						FCanExecuteAction(),
+						FGetActionCheckState::CreateSP(SettingsEditor.Get(), &SSettingsEditor::IsSearchFilterActive, SearchTerm, /** Equals */true)
+					),
+					EUserInterfaceActionType::RadioButton
+				);
+			}
+		}),
+		FilterAction,
+		EUserInterfaceActionType::ToggleButton
+	);
+
+	FilterEntry.Icon = FSlateIcon(FAppStyle::GetAppStyleSetName(), "Icons.Filter");
+
+	return FilterEntry;
 }
 
-void SSettingsEditor::OnSearchFilterButtonClicked(ECheckBoxState InState, FString InFilter)
+void SSettingsEditor::ToggleSearchFilterAction(FString InFilter)
 {
-	if (InState == ECheckBoxState::Checked)
+	if (IsSearchFilterActive(InFilter, /** Equals */false) == ECheckBoxState::Unchecked)
 	{
 		SearchBox->SetText(FText::FromString(InFilter + TEXT("=")));
 	}
@@ -931,6 +998,18 @@ void SSettingsEditor::OnSearchFilterButtonClicked(ECheckBoxState InState, FStrin
 	}
 }
 
+ECheckBoxState SSettingsEditor::IsSearchFilterActive(FString InFilter, bool bInStrictCompare)
+{
+	const FString ActiveSearch = SearchBox->GetText().ToString();
+	const bool bSearchFilterActive = bInStrictCompare ? ActiveSearch.Equals(InFilter) : ActiveSearch.StartsWith(InFilter);
+	return bSearchFilterActive ? ECheckBoxState::Checked : ECheckBoxState::Unchecked;
+}
+
+void SSettingsEditor::SetSearchFilter(FString InSearchTerm)
+{
+	SearchBox->SetText(FText::FromString(InSearchTerm));
+}
+
 bool SSettingsEditor::HandleSettingsViewEnabled() const
 {
 	ISettingsSectionPtr SelectedSection = Model->GetSelectedSection();
diff --git a/Engine/Source/Developer/SettingsEditor/Private/Widgets/SSettingsEditor.h b/Engine/Source/Developer/SettingsEditor/Private/Widgets/SSettingsEditor.h
--- a/Engine/Source/Developer/SettingsEditor/Private/Widgets/SSettingsEditor.h
+++ b/Engine/Source/Developer/SettingsEditor/Private/Widgets/SSettingsEditor.h
@@ -11,7 +11,6 @@
 #include "Layout/Visibility.h"
 #include "Misc/NotifyHook.h"
 #include "Widgets/DeclarativeSyntaxSupport.h"
-#include "Widgets/SWidget.h"
 #include "Widgets/SCompoundWidget.h"
 #include "Widgets/SOverlay.h"
 
@@ -19,8 +18,12 @@ class IDetailsView;
 class SAssetSearchBox;
 class SScrollBox;
 class SVerticalBox;
+class USettingsEditorMenuContext;
+class UToolMenu;
+enum class ECheckBoxState : uint8;
 struct FAssetSearchBoxSuggestion;
 struct FPropertyChangedEvent;
+struct FToolMenuEntry;
 
 /**
  * Implements an editor widget for settings.
@@ -88,6 +91,8 @@ protected:
 
 
 private:
+	static void FillSettingsFiltersSection(UToolMenu* InMenu);
+	static FToolMenuEntry CreateFiltersCheckboxEntry(const USettingsEditorMenuContext* InContext, const FString& InKey, const TSet<FString>& InValues);
 
 	/** Callback for when the user's culture has changed. */
 	void HandleCultureChanged();
@@ -133,10 +138,10 @@ private:
 	void OnSearchTextCommitted(const FText& InText, ETextCommit::Type InType);
 	void OnAssetSearchSuggestionFilter(const FText& InText, TArray<FAssetSearchBoxSuggestion>& OutSuggestion, FText& OutText);
 	FText OnAssetSearchSuggestionChosen(const FText& InText, const FString& InSuggestion);
-	void HandleSearchTermsChanged(const ISettingsEditorModule::FSearchTerm& InNewSearchTerm);
-	
-	ECheckBoxState IsSearchFilterButtonActive(FString InFilter) const;
-	void OnSearchFilterButtonClicked(ECheckBoxState InState, FString InFilter);
+
+	void ToggleSearchFilterAction(FString InFilter);
+	ECheckBoxState IsSearchFilterActive(FString InFilter, bool bInStrictCompare);
+	void SetSearchFilter(FString InSearchTerm);
 
 	/** Holds the vertical box for settings categories. */
 	TSharedPtr<SVerticalBox> CategoriesBox;
@@ -164,10 +169,4 @@ private:
 
 	/** Actual search box with suggestions */
 	TSharedPtr<SAssetSearchBox> SearchBox;
-
-	/** Used to store quick filters checkboxes representing the keys of the suggestions (cvar, keybind, ...) */
-	TSharedPtr<SScrollBox> SearchFiltersBox;
-	
-	/** Used to keep track of the search filters added */
-	TSet<FString> QuickSearchFiltersAdded;
 };
diff --git a/Engine/Source/Developer/SettingsEditor/Private/Widgets/SSettingsSectionHeader.cpp b/Engine/Source/Developer/SettingsEditor/Private/Widgets/SSettingsSectionHeader.cpp
--- a/Engine/Source/Developer/SettingsEditor/Private/Widgets/SSettingsSectionHeader.cpp
+++ b/Engine/Source/Developer/SettingsEditor/Private/Widgets/SSettingsSectionHeader.cpp
@@ -117,23 +117,23 @@ void SSettingsSectionHeader::Construct(const FArguments& InArgs, const UObject*
 					.AutoWidth()
 					.Padding(8.0f, 0.0f, 0.0f, 0.0f)
 					[
-						// export button
+						// import button
 						SNew(SButton)
-						.IsEnabled(this, &SSettingsSectionHeader::HandleExportButtonEnabled)
-						.OnClicked(this, &SSettingsSectionHeader::HandleExportButtonClicked)
-						.Text(LOCTEXT("ExportButtonText", "Export..."))
-						.ToolTipText(LOCTEXT("ExportButtonTooltip", "Export these settings to a file on your computer"))
+						.IsEnabled(this, &SSettingsSectionHeader::HandleImportButtonEnabled)
+						.OnClicked(this, &SSettingsSectionHeader::HandleImportButtonClicked)
+						.Text(LOCTEXT("ImportButtonText", "Import..."))
+						.ToolTipText(LOCTEXT("ImportButtonTooltip", "Import these settings from a file on your computer"))
 					]
 					+SHorizontalBox::Slot()
 					.AutoWidth()
 					.Padding(8.0f, 0.0f, 0.0f, 0.0f)
 					[
-						// import button
+						// export button
 						SNew(SButton)
-						.IsEnabled(this, &SSettingsSectionHeader::HandleImportButtonEnabled)
-						.OnClicked(this, &SSettingsSectionHeader::HandleImportButtonClicked)
-						.Text(LOCTEXT("ImportButtonText", "Import..."))
-						.ToolTipText(LOCTEXT("ImportButtonTooltip", "Import these settings from a file on your computer"))
+						.IsEnabled(this, &SSettingsSectionHeader::HandleExportButtonEnabled)
+						.OnClicked(this, &SSettingsSectionHeader::HandleExportButtonClicked)
+						.Text(LOCTEXT("ExportButtonText", "Export..."))
+						.ToolTipText(LOCTEXT("ExportButtonTooltip", "Export these settings to a file on your computer"))
 					]
 					+SHorizontalBox::Slot()
 					.AutoWidth()
diff --git a/Engine/Source/Editor/EditorStyle/Private/StarshipStyle.cpp b/Engine/Source/Editor/EditorStyle/Private/StarshipStyle.cpp
--- a/Engine/Source/Editor/EditorStyle/Private/StarshipStyle.cpp
+++ b/Engine/Source/Editor/EditorStyle/Private/StarshipStyle.cpp
@@ -2731,6 +2731,27 @@ void FStarshipEditorStyle::FStyle::SetupViewportStyles()
 			Set("ViewportToolbar.SplitComboButton", SplitComboButton);
 		}
 
+		// Variant with no background
+		{
+			FToolBarStyle ViewportToolbarNoBackground(ViewportToolbarStyle);
+			ViewportToolbarNoBackground.SetBackground(FSlateNoResource());
+			ViewportToolbarNoBackground.BackgroundPadding = FMargin(0.f);
+			Set("NoBackgroundViewportToolbar", ViewportToolbarNoBackground);
+			Set("NoBackgroundViewportToolbar.Button.Start", ButtonStyleStart);
+            Set("NoBackgroundViewportToolbar.Button.Middle", ButtonStyleMiddle);
+            Set("NoBackgroundViewportToolbar.Button.End", ButtonStyleEnd);
+            
+            Set("NoBackgroundViewportToolbar.ToggleButton.Start", ToggleButtonStyleStart);
+            Set("NoBackgroundViewportToolbar.ToggleButton.Middle", ToggleButtonStyleMiddle);
+            Set("NoBackgroundViewportToolbar.ToggleButton.End", ToggleButtonStyleEnd);
+            
+            // Ensure split buttons don't have rounded corners on the inside
+            Set("NoBackgroundViewportToolbar.SplitToggleButton", ToggleButtonStyleStart);
+            FComboButtonStyle SplitComboButton = ViewportToolbarStyle.ComboButtonStyle;
+            SplitComboButton.ButtonStyle = ButtonStyleEnd;
+            Set("NoBackgroundViewportToolbar.SplitComboButton", SplitComboButton);
+		}
+
 		// Special styling for top-level raised buttons.
 		FToolBarStyle ViewportToolbarRaisedStyle = ViewportToolbarStyle;
 		{
diff --git a/Engine/Source/Editor/InputBindingEditor/Private/InputBindingEditorModule.cpp b/Engine/Source/Editor/InputBindingEditor/Private/InputBindingEditorModule.cpp
--- a/Engine/Source/Editor/InputBindingEditor/Private/InputBindingEditorModule.cpp
+++ b/Engine/Source/Editor/InputBindingEditor/Private/InputBindingEditorModule.cpp
@@ -229,21 +229,29 @@ public:
 		FInputBindingManager& InputBindingManager = FInputBindingManager::Get();
 		
 		ISettingsEditorModule* SettingsEditorModule = FModuleManager::GetModulePtr<ISettingsEditorModule>(TEXT("SettingsEditor"));
-		const FString SearchKey = TEXT("keybind");
 		FName SearchContext = NAME_None;
 		if (TSharedPtr<IDetailsView> DetailsView = DetailBuilder->GetDetailsViewSharedPtr())
 		{
 			SearchContext = DetailsView->GetIdentifier();
 		}
 
+		static const FText SearchKeyKeybindUnbound = LOCTEXT("SearchKeyKeybindUnbound", "Unbound");
+		static const FText SearchKeyKeybind = LOCTEXT("SearchKeyKeybind", "KeyBind");
+		static const FText SearchKeyContext = LOCTEXT("SearchKeyContext", "Context");
+		static const FTextFormat TokenCombineFormat = INVTEXT("{0} {1}");
+		static const FTextFormat TokenDelimiter = INVTEXT("\"{0}\"");
+		static const FTextFormat SearchFormat = INVTEXT("{0}={1}");
 		for (TSharedPtr<FChordTreeItem>& TreeItem : ContextList)
 		{
 			check(TreeItem->IsContext());
 
-			IDetailCategoryBuilder& CategoryBuilder = DetailBuilder->EditCategory(TreeItem->GetBindingContext()->GetContextName(), TreeItem->GetBindingContext()->GetContextDesc());
+			FName BindingContextName = TreeItem->GetBindingContext()->GetContextName();
+			const FText& BindingContextLabel = TreeItem->GetBindingContext()->GetContextDesc();
+
+			IDetailCategoryBuilder& CategoryBuilder = DetailBuilder->EditCategory(BindingContextName, BindingContextLabel);
 
 			TArray<TSharedPtr<FUICommandInfo>> Commands;
-			InputBindingManager.GetCommandInfosFromContext(TreeItem->GetBindingContext()->GetContextName(), Commands);
+			InputBindingManager.GetCommandInfosFromContext(BindingContextName, Commands);
 
 			// Filter to allowed bindings
 			Commands.RemoveAll([&InputBindingManager](const TSharedPtr<FUICommandInfo>& CommandInfo)
@@ -254,7 +262,6 @@ public:
 			Commands.Sort(FChordSort(true, false));
 
 			TMap<FName, IDetailGroup*> BundleMap;
-
 			for(TSharedPtr<FUICommandInfo>& CommandInfo : Commands)
 			{
 				FDetailWidgetRow* Row = nullptr;
@@ -287,23 +294,19 @@ public:
 					Row = &CategoryBuilder.AddCustomRow(CommandInfo->GetLabel());
 				}
 
-				// Set up search filter (for i.e. keybind=F)
-				const TSharedRef<const FInputChord> FirstInputChord = CommandInfo->GetActiveChord(EMultipleKeyBindingIndex::Primary);
-				const TSharedRef<const FInputChord> SecondInputChord = CommandInfo->GetActiveChord(EMultipleKeyBindingIndex::Secondary);
-
-				const bool bFirstChordValid = FirstInputChord->IsValidChord();
-				const bool bSecondChordValid = SecondInputChord->IsValidChord();
-
-				if (bFirstChordValid || bSecondChordValid)
 				{
-					static const FTextFormat KeyFormat = LOCTEXT("SearchKeyFilter", "keybind={0}");
-					static const FTextFormat TokenCombineFormat = INVTEXT("{0} {1}");
+					// Set up search filter (for i.e. keybind="F")
+					const TSharedRef<const FInputChord> FirstInputChord = CommandInfo->GetActiveChord(EMultipleKeyBindingIndex::Primary);
+					const TSharedRef<const FInputChord> SecondInputChord = CommandInfo->GetActiveChord(EMultipleKeyBindingIndex::Secondary);
+
+					const bool bFirstChordValid = FirstInputChord->IsValidChord();
+					const bool bSecondChordValid = SecondInputChord->IsValidChord();
 
 					FText FirstKeyFilter, SecondKeyFilter;
 					if (bFirstChordValid)
 					{
-						FText InputText = FirstInputChord->GetInputText();
-						FirstKeyFilter = FText::FormatOrdered(KeyFormat, InputText);
+						const FText InputText = FText::FormatOrdered(TokenDelimiter, FirstInputChord->GetInputText());
+						FirstKeyFilter = FText::FormatOrdered(SearchFormat, SearchKeyKeybind.ToLower(), InputText);
 
 						// Register search key and value for settings editor suggestions
 						if (SettingsEditorModule)
@@ -311,7 +314,7 @@ public:
 							ISettingsEditorModule::FSearchTerm SearchTerm
 							{
 								.ContainerName = SearchContext,
-								.Key = SearchKey,
+								.Key = SearchKeyKeybind.ToString(),
 								.Value = InputText.ToString()
 							};
 							SettingsEditorModule->RegisterSearchTerm(SearchTerm);
@@ -320,8 +323,8 @@ public:
 
 					if (bSecondChordValid)
 					{
-						FText InputText = SecondInputChord->GetInputText();
-						SecondKeyFilter = FText::FormatOrdered(KeyFormat, InputText);
+						const FText InputText = FText::FormatOrdered(TokenDelimiter, SecondInputChord->GetInputText());
+						SecondKeyFilter = FText::FormatOrdered(SearchFormat, SearchKeyKeybind.ToLower(), InputText);
 
 						// Register search key and value for settings editor suggestions
 						if (SettingsEditorModule)
@@ -329,7 +332,7 @@ public:
 							ISettingsEditorModule::FSearchTerm SearchTerm
 							{
 								.ContainerName = SearchContext,
-								.Key = SearchKey,
+								.Key = SearchKeyKeybind.ToString(),
 								.Value = InputText.ToString()
 							};
 							SettingsEditorModule->RegisterSearchTerm(SearchTerm);
@@ -340,12 +343,50 @@ public:
 						? FText::FormatOrdered(TokenCombineFormat, FirstKeyFilter, SecondKeyFilter)
 						: (bFirstChordValid ? FirstKeyFilter : SecondKeyFilter);
 
+					// Allow search for unbound commands like this : keybind=unbound
+					if (FilterText.IsEmpty())
+					{
+						FilterText = FText::FormatOrdered(SearchFormat, SearchKeyKeybind.ToLower(), SearchKeyKeybindUnbound.ToLower());
+
+						// Register search key and value for settings editor suggestions
+						if (SettingsEditorModule)
+						{
+							ISettingsEditorModule::FSearchTerm SearchTerm
+							{
+								.ContainerName = SearchContext,
+								.Key = SearchKeyKeybind.ToString(),
+								.Value = SearchKeyKeybindUnbound.ToString()
+							};
+							SettingsEditorModule->RegisterSearchTerm(SearchTerm);
+						}
+					}
+
 					// Command label and similar string might already be stored
 					if (!Row->FilterTextString.IsEmpty())
 					{
 						FilterText = FText::FormatOrdered(TokenCombineFormat, Row->FilterTextString, FilterText);
 					}
 
+					// Append command binding context
+					{
+						FText BindingContext = FText::FormatOrdered(TokenDelimiter, FText::FromName(BindingContextName));
+						
+						// Register search key and value for settings editor suggestions
+						if (SettingsEditorModule)
+						{
+							ISettingsEditorModule::FSearchTerm SearchTerm
+							{
+								.ContainerName = SearchContext,
+								.Key = SearchKeyContext.ToString(),
+								.Value = BindingContext.ToString()
+							};
+							SettingsEditorModule->RegisterSearchTerm(SearchTerm);
+						}
+
+						BindingContext = FText::FormatOrdered(SearchFormat, SearchKeyContext.ToLower(), BindingContext);
+						FilterText = FText::FormatOrdered(TokenCombineFormat, FilterText, BindingContext);
+					}
+
 					Row->FilterString(FilterText);
 				}
 
