Adding ability to use ShowOnlyInnerProperties on per-platform settings customization

#rnx

#Synced-CL 50410508

diff --git a/Engine/Source/Developer/SettingsEditor/Private/Customizations/PerPlatformSettingsCustomization.cpp b/Engine/Source/Developer/SettingsEditor/Private/Customizations/PerPlatformSettingsCustomization.cpp
--- a/Engine/Source/Developer/SettingsEditor/Private/Customizations/PerPlatformSettingsCustomization.cpp
+++ b/Engine/Source/Developer/SettingsEditor/Private/Customizations/PerPlatformSettingsCustomization.cpp
@@ -20,13 +20,19 @@ FPerPlatformSettingsCustomization::FPerPlatformSettingsCustomization()
 
 void FPerPlatformSettingsCustomization::CustomizeHeader(TSharedRef<IPropertyHandle> PropertyHandle, FDetailWidgetRow& HeaderRow, IPropertyTypeCustomizationUtils& CustomizationUtils)
 {
-	HeaderRow
-	.WholeRowContent()
-	.MaxDesiredWidth(TOptional<float>())
-	.MaxDesiredWidth(TOptional<float>())
-	[
-		PropertyHandle->CreatePropertyNameWidget()
-	];
+
+	static const FName ShowOnlyInners("ShowOnlyInnerProperties");
+	const bool bShowHeader = !PropertyHandle->HasMetaData(ShowOnlyInners);
+	if (bShowHeader)
+	{
+		HeaderRow
+			.WholeRowContent()
+			.MaxDesiredWidth(TOptional<float>())
+			.MaxDesiredWidth(TOptional<float>())
+			[
+				PropertyHandle->CreatePropertyNameWidget()
+			];
+	}
 }
 
 void FPerPlatformSettingsCustomization::CustomizeChildren(TSharedRef<IPropertyHandle> PropertyHandle, IDetailChildrenBuilder& ChildBuilder, IPropertyTypeCustomizationUtils& CustomizationUtils)
