EditorUsability : SettingsEditor
- Improved settings editor search following feedback (part 3)
- UI reworked to match the rest of the editor (treeview, splitter, category, style)
- Improved search to only apply to the selected section instead of showing all results
- Improved Filters and suggestions to be dynamic and (dis)allow specific sections (they only appear in the proper context)
- Added keybind action button to search for specific keys pressed
- Allow partial search match from begin or end
- Added icons for keybind and cvar filters

#Synced-CL 50871375

#Synced-CL 50871375

diff --git a/Engine/Source/Developer/SettingsEditor/Private/Customizations/SettingsEditorCustomization.cpp b/Engine/Source/Developer/SettingsEditor/Private/Customizations/SettingsEditorCustomization.cpp
--- a/Engine/Source/Developer/SettingsEditor/Private/Customizations/SettingsEditorCustomization.cpp
+++ b/Engine/Source/Developer/SettingsEditor/Private/Customizations/SettingsEditorCustomization.cpp
@@ -75,17 +75,32 @@ void FSettingsEditorCustomization::CustomizeDetails(IDetailLayoutBuilder& InDeta
 {
 	TArray<TWeakObjectPtr<UObject>> SettingsObjects;
 	InDetailLayoutBuilder.GetObjectsBeingCustomized(SettingsObjects);
-	
-	ISettingsEditorModule* SettingsEditorModule = FModuleManager::GetModulePtr<ISettingsEditorModule>(TEXT("SettingsEditor"));
-	FName SearchContext = NAME_None;
-	if (const TSharedPtr<IDetailsView> DetailsView = InDetailLayoutBuilder.GetDetailsViewSharedPtr())
+
+	static const FTextFormat TokenCombineFormat = INVTEXT("{0} {1}");
+	static const FText SearchKeyCVar = LOCTEXT("SearchKeyCVar", "CVar:");
+	static const FTextFormat TokenDelimiter = INVTEXT("\"{0}\""); // Use for strict search instead of partial...
+
+	ISettingsEditorModule::FSearchTerm* CVarSearchTerm = nullptr;
+	if (ISettingsEditorModule* SettingsEditorModule = FModuleManager::GetModulePtr<ISettingsEditorModule>(TEXT("SettingsEditor")))
 	{
-		SearchContext = DetailsView->GetIdentifier();
+		if (const TSharedPtr<IDetailsView> DetailsView = InDetailLayoutBuilder.GetDetailsViewSharedPtr())
+		{
+			const FName SearchContext = DetailsView->GetIdentifier();
+			if (SettingsEditorModule->RegisterSearchTerm(SearchContext, SearchKeyCVar.ToString()))
+			{
+				CVarSearchTerm = SettingsEditorModule->FindSearchTerm(SearchContext, SearchKeyCVar.ToString());
+				CVarSearchTerm->Label = LOCTEXT("SearchKeyCVarLabel", "CVar");
+				CVarSearchTerm->Tooltip = LOCTEXT("SearchKeyCVarTooltip", "Filters properties to show those with a Console Variable");
+				CVarSearchTerm->IconStyleSetName = "OutputLogStyle";
+				CVarSearchTerm->IconStyleName = "DebugConsole.Icon";
+				CVarSearchTerm->DisallowedSections.Add(TEXT("InputBindings")); // Keyboard Shortcuts
+				CVarSearchTerm->AllowedSections.Add(NAME_None); // All section
+				// No need to show filter buttons in Editor settings
+				CVarSearchTerm->bShowFilter = !SearchContext.IsEqual("Editor");
+			}
+		}
 	}
 
-	static const FText SearchKeyCVar = LOCTEXT("SearchKeyCVar", "CVar");
-	static const FTextFormat SearchFormat = INVTEXT("{0}={1}");
-	static const FTextFormat TokenDelimiter = INVTEXT("\"{0}\"");
 	for (const TWeakObjectPtr<UObject>& SettingsObject : SettingsObjects)
 	{
 		TStrongObjectPtr<UObject> PinnedSettingsObject = SettingsObject.Pin();
@@ -115,7 +130,7 @@ void FSettingsEditorCustomization::CustomizeDetails(IDetailLayoutBuilder& InDeta
 			}
 			
 			const FString& CVar = SettingsProperty->GetMetaData(TEXT("ConsoleVariable"));
-			if (!CVar.IsEmpty())
+			if (CVarSearchTerm != nullptr && !CVar.IsEmpty())
 			{
 				FDetailWidgetRow& Row = PropertyRow->CustomWidget(true);
 				Row.NameContent()
@@ -138,18 +153,19 @@ void FSettingsEditorCustomization::CustomizeDetails(IDetailLayoutBuilder& InDeta
 				FText FilterText = FText::FormatOrdered(TokenDelimiter, FText::FromString(CVar));
 
 				// Register search key and value for settings editor suggestions
-				if (SettingsEditorModule)
+				CVarSearchTerm->Values.Add(CVar);
+				
+				// Only allow this search term in the section specified by the settings object name
 				{
-					ISettingsEditorModule::FSearchTerm SearchTerm
+					FString NativeDisplayName = PinnedSettingsObject->GetClass()->GetMetaData("DisplayName");
+					if (NativeDisplayName.IsEmpty())
 					{
-						.ContainerName = SearchContext,
-						.Key = SearchKeyCVar.ToString(),
-						.Value = FilterText.ToString()
-					};
-					SettingsEditorModule->RegisterSearchTerm(SearchTerm);
+						NativeDisplayName = FName::NameToDisplayString(PinnedSettingsObject->GetClass()->GetName(), false);
+					}
+					CVarSearchTerm->AllowedSections.Add(FName(NativeDisplayName));
 				}
 
-				FilterText = FText::FormatOrdered(SearchFormat, SearchKeyCVar.ToLower(), FilterText);
+				FilterText = FText::FormatOrdered(TokenCombineFormat, SearchKeyCVar.ToLower(), FilterText);
 				Row.FilterString(FilterText);
 			}
 			else
diff --git a/Engine/Source/Developer/SettingsEditor/Private/SettingsEditorModule.cpp b/Engine/Source/Developer/SettingsEditor/Private/SettingsEditorModule.cpp
--- a/Engine/Source/Developer/SettingsEditor/Private/SettingsEditorModule.cpp
+++ b/Engine/Source/Developer/SettingsEditor/Private/SettingsEditorModule.cpp
@@ -196,32 +196,58 @@ public:
 	{
 		ShouldRegisterSettingsDelegate = MoveTemp(InShouldRegisterSettingDelegate);
 	}
-	
-	virtual void RegisterSearchTerm(const FSearchTerm& InSearchTerm) override
+
+	virtual bool RegisterSearchTerm(FName InContainerName, const FString& InKey) override
 	{
-		if (InSearchTerm.ContainerName.IsNone())
+		if (InContainerName.IsNone() || InKey.IsEmpty())
 		{
-			return;
+			return false;
 		}
 
-		bool bAlreadyRegistered = false;
-		SearchTerms
-			.FindOrAdd(InSearchTerm.ContainerName)
-			.FindOrAdd(InSearchTerm.Key)
-			.Add(InSearchTerm.Value, &bAlreadyRegistered);
-
-		if (!bAlreadyRegistered)
+		// Already registered ?
+		if (FSearchTerm* SearchTerm = FindSearchTerm(InContainerName, InKey))
 		{
-			OnSearchTermsChangedDelegate.Broadcast(InSearchTerm);
+			return true;
 		}
+
+		TMap<FString, FSearchTerm>& ContainerKeys = SearchTerms.FindOrAdd(InContainerName);
+		ContainerKeys.FindOrAdd(InKey);
+
+		// Resort to display in alphabetical order
+		ContainerKeys.KeySort(TLess<FString>());
+
+		// Notify
+		OnSearchTermsChangedDelegate.Broadcast(InContainerName, InKey);
+
+		return true;
 	}
 	
 	virtual FOnSearchTermsChanged::RegistrationType& OnSearchTermsChanged() override
 	{
 		return OnSearchTermsChangedDelegate;
 	}
+
+	virtual FSearchTerm* FindSearchTerm(FName InContainerName, const FString& InKey) override
+	{
+		if (TMap<FString, FSearchTerm>* ContainerSearchTerms = SearchTerms.Find(InContainerName))
+		{
+			return ContainerSearchTerms->Find(InKey);
+		}
+
+		return nullptr;
+	}
 	
-	virtual const TMap<FString, TSet<FString>>* FindSearchTerms(FName InContainerName) const override
+	virtual const FSearchTerm* FindSearchTerm(FName InContainerName, const FString& InKey) const override
+	{
+		if (const TMap<FString, FSearchTerm>* ContainerSearchTerms = SearchTerms.Find(InContainerName))
+		{
+			return ContainerSearchTerms->Find(InKey);
+		}
+
+		return nullptr;
+	}
+
+	virtual const TMap<FString, FSearchTerm>* FindSearchTerms(FName InContainerName) const override
 	{
 		return SearchTerms.Find(InContainerName);
 	}
@@ -538,9 +564,9 @@ private:
 	FShouldRegisterSettingsDelegate ShouldRegisterSettingsDelegate;
 	
 	/** Used to store Context->Key->Values for search suggestions */
-	TMap<FName, TMap<FString, TSet<FString>>> SearchTerms;
-	
-	/** Delegate called when search suggestions changed */
+	TMap<FName, TMap<FString, FSearchTerm>> SearchTerms;
+
+	/** Used to broadcast change events */
 	FOnSearchTermsChanged OnSearchTermsChangedDelegate;
 };
 
diff --git a/Engine/Source/Developer/SettingsEditor/Private/Widgets/SSettingsEditor.cpp b/Engine/Source/Developer/SettingsEditor/Private/Widgets/SSettingsEditor.cpp
--- a/Engine/Source/Developer/SettingsEditor/Private/Widgets/SSettingsEditor.cpp
+++ b/Engine/Source/Developer/SettingsEditor/Private/Widgets/SSettingsEditor.cpp
@@ -10,17 +10,14 @@
 #include "Layout/WidgetPath.h"
 #include "Framework/Application/SlateApplication.h"
 #include "Widgets/Layout/SSeparator.h"
-#include "Widgets/Layout/SSpacer.h"
-#include "Widgets/Images/SImage.h"
 #include "Widgets/Layout/SScrollBox.h"
 #include "Styling/AppStyle.h"
 #include "AnalyticsEventAttribute.h"
 #include "DesktopPlatformModule.h"
 #include "EngineAnalytics.h"
-#include "Interfaces/IAnalyticsProvider.h"
-#include "Widgets/Text/STextBlock.h"
 #include "PropertyEditorModule.h"
 #include "IDetailsView.h"
+#include "SSettingsEditorCategoryTree.h"
 #include "Widgets/Input/SHyperlink.h"
 #include "Widgets/SSettingsSectionHeader.h"
 #include "SSettingsEditorCheckoutNotice.h"
@@ -36,17 +33,29 @@
 #include "Misc/ConfigCacheIni.h"
 #include "Misc/MessageDialog.h"
 #include "Models/SettingsEditorMenuContext.h"
+#include "Widgets/Layout/SBox.h"
 
 #define LOCTEXT_NAMESPACE "SSettingsEditor"
 
+FLazyName SSettingsEditor::SettingsEditorToolbarName = TEXT("SettingsEditor.SearchToolbar");
+const FTextFormat SSettingsEditor::TokenCombineFormat = INVTEXT("{0} {1}");
+const FTextFormat SSettingsEditor::TokenDelimiter = INVTEXT("\"{0}\"");
+
 /* SSettingsEditor structors
  *****************************************************************************/
 
 SSettingsEditor::~SSettingsEditor()
 {
 	Model->OnSelectionChanged().RemoveAll(this);
-	SettingsContainer->OnCategoryModified().RemoveAll(this);
-	FInternationalization::Get().OnCultureChanged().RemoveAll(this);
+	if (ISettingsEditorModule* SettingsEditorModule = FModuleManager::GetModulePtr<ISettingsEditorModule>(TEXT("SettingsEditor")))
+	{
+		SettingsEditorModule->OnSearchTermsChanged().RemoveAll(this);
+	}
+	if (MenuToolbarUpdateHandle.IsValid())
+	{
+		FTSTicker::GetCoreTicker().RemoveTicker(MenuToolbarUpdateHandle);
+		MenuToolbarUpdateHandle.Reset();
+	}
 }
 
 
@@ -55,7 +64,6 @@ SSettingsEditor::~SSettingsEditor()
 
 void SSettingsEditor::Construct( const FArguments& InArgs, const ISettingsEditorModelRef& InModel )
 {
-	bIsActiveTimerRegistered = false;
 	Model = InModel;
 	SettingsContainer = InModel->GetSettingsContainer();
 	OnApplicationRestartRequiredDelegate = InArgs._OnApplicationRestartRequired;
@@ -94,10 +102,10 @@ void SSettingsEditor::Construct( const FArguments& InArgs, const ISettingsEditor
 	TSharedPtr<FSettingsDetailRootObjectCustomization> RootObjectCustomization = MakeShareable(new FSettingsDetailRootObjectCustomization(Model, SettingsView.ToSharedRef()));
 	RootObjectCustomization->Initialize();
 	SettingsView->SetRootObjectCustomizationInstance(RootObjectCustomization);
-	
-	const FName SettingsEditorToolbarName = TEXT("SettingsEditor.SearchToolbar");
+
+	// Register menu toolbar
 	UToolMenus* ToolMenus = UToolMenus::Get();
-	if (!ToolMenus->IsMenuRegistered(SettingsEditorToolbarName))
+	if (IsValid(ToolMenus) && !ToolMenus->IsMenuRegistered(SettingsEditorToolbarName))
 	{
 		UToolMenu* const SettingsToolbarMenu = ToolMenus->RegisterMenu(
 			SettingsEditorToolbarName, NAME_None /* parent */, EMultiBoxType::SlimHorizontalToolBar
@@ -111,75 +119,59 @@ void SSettingsEditor::Construct( const FArguments& InArgs, const ISettingsEditor
 		);
 	}
 
-	USettingsEditorMenuContext* MenuContext = NewObject<USettingsEditorMenuContext>();
-	MenuContext->SettingsEditorWeak = SharedThis(this);
-	const TSharedRef<SWidget> MenuWidget = ToolMenus->GenerateWidget(SettingsEditorToolbarName, FToolMenuContext(MenuContext));
-
-	TSharedPtr<SScrollBox> ScrollBox;
+	TSharedPtr<SBorder> CategoryTreeBorder;
 	ChildSlot
 	[
-		SNew(SVerticalBox)
-
-		+ SVerticalBox::Slot()
-		.FillHeight(1.0f)
-		.Padding(16.0f, 0.0f)
+		SNew(SBorder)
+		.BorderImage(FAppStyle::Get().GetBrush("Brushes.Panel"))
+		.Padding(0.f)
 		[
-			SNew(SHorizontalBox)
-			+ SHorizontalBox::Slot()
-			.AutoWidth()
-			.Padding(0.0f, 16.0f)
-			[
-				// categories menu
-				SAssignNew(ScrollBox, SScrollBox)
-				+ SScrollBox::Slot()
-				[
-					SNew(SHorizontalBox)
-					+ SHorizontalBox::Slot()
-					.AutoWidth()
-					[
-						SAssignNew(CategoriesBox, SVerticalBox)
-					]
-
-					+ SHorizontalBox::Slot()
-					.AutoWidth()
-					[
-						SNew(SSpacer)
-						.Size(FVector2D(24.0f, 0.0f))
-					]
-				]
-			]
-			+ SHorizontalBox::Slot()
-			.AutoWidth()
-			.Padding(FMargin(24.0f, 0.0f, 24.0f, 0.0f))
+			SNew(SSplitter)
+			.Style(FAppStyle::Get(), "SplitterPanel")
+			+ SSplitter::Slot()
+			.Value(.35f)
 			[
-				SNew(SSeparator)
-				.Orientation(Orient_Vertical)
+				SAssignNew(CategoryTreeBorder, SBorder)
+				.VAlign(VAlign_Fill)
+				.BorderImage(FAppStyle::Get().GetBrush("SettingsEditor.ListBorder"))
+				.Padding(8.f)
 			]
-			+ SHorizontalBox::Slot()
-			.FillWidth(1.0f)
-			.Padding(0.0f, 16.0f)
+			+ SSplitter::Slot()
 			[
 				SNew(SVerticalBox)
-				.Visibility(this, &SSettingsEditor::HandleSettingsBoxVisibility)
 				+ SVerticalBox::Slot()
+				.Padding(4.f)
 				.AutoHeight()
 				[
 					SNew(SHorizontalBox)
 					+ SHorizontalBox::Slot()
+					.AutoWidth()
+					.VAlign(VAlign_Center)
+					.Padding(0.0f)
+					[
+						SAssignNew(MenuToolbarBox, SBox)
+						.MinDesiredHeight(24.f)
+						[
+							// Menu toolbar goes here
+							SNullWidget::NullWidget
+						]
+					]
+					+ SHorizontalBox::Slot()
 					.FillWidth(1.0f)
 					.VAlign(VAlign_Center)
-					.Padding(0.f, 0.f, 0.f, 0.f)
+					.Padding(0.f)
 					[
 						SAssignNew(SearchBox, SAssetSearchBox)
-						.SuggestionListHeight(50.f)
-						.HintText(LOCTEXT("FilterSearch", "Search..."))
-						.ShowSearchHistory(true)
-						.ToolTipText(this, &SSettingsEditor::GetSearchToolTipText)
+						.SuggestionListHeight(60.f)
+						.ShowSearchHistory(false) // Only suggest, otherwise its confusing
+						.DelayChangeNotificationsWhileTyping(true)
+						.ToolTipText(FText::GetEmpty()) // Do not show any tooltip, it's showing when suggestions are listed...
+						.HintText(this, &SSettingsEditor::GetSearchHintText)
 						.OnTextChanged(this, &SSettingsEditor::OnSearchTextChanged)
 						.OnTextCommitted(this, &SSettingsEditor::OnSearchTextCommitted)
 						.OnAssetSearchBoxSuggestionFilter(this, &SSettingsEditor::OnAssetSearchSuggestionFilter)
 						.OnAssetSearchBoxSuggestionChosen(this, &SSettingsEditor::OnAssetSearchSuggestionChosen)
-						.DelayChangeNotificationsWhileTyping(true)
+						.OnKeyDownHandler(this, &SSettingsEditor::OnSearchKeyDownHandler)
 						.AddMetaData<FTagMetaData>(FTagMetaData(TEXT("SettingsEditorSearchBox")))
 					]
 					+ SHorizontalBox::Slot()
@@ -190,29 +182,24 @@ void SSettingsEditor::Construct( const FArguments& InArgs, const ISettingsEditor
 						// Import button
 						SNew(SButton)
 						.OnClicked(this, &SSettingsEditor::HandleImportButtonClicked)
-						.Text(LOCTEXT("ImportButtonText", "Import..."))
+						.Text(LOCTEXT("ImportAllButtonText", "Import All..."))
 						.ToolTipText(LOCTEXT("ImportAllButtonTooltip", "Import all editor settings from a file on your computer"))
 					]
 					+ SHorizontalBox::Slot()
 					.AutoWidth()
 					.VAlign(VAlign_Center)
-					.Padding(8.0f, 0.0f, 6.0f, 0.0f)
+					.Padding(8.0f, 0.0f)
 					[
 						// Export button
 						SNew(SButton)
 						.OnClicked(this, &SSettingsEditor::HandleExportButtonClicked)
-						.Text(LOCTEXT("ExportButtonText", "Export..."))
+						.Text(LOCTEXT("ExportAllButtonText", "Export All..."))
 						.ToolTipText(LOCTEXT("ExportAllButtonTooltip", "Export all editor settings to a file on your computer"))
 					]
 				]
 				+ SVerticalBox::Slot()
-				.Padding(0.f, 4.f)
-				.AutoHeight()
-				[
-					MenuWidget
-				]
-				+ SVerticalBox::Slot()
-				.FillHeight(1.0f)
+				.Padding(4.f)
+				.FillHeight(1.f)
 				[
 					// settings area
 					SNew(SOverlay)
@@ -227,13 +214,23 @@ void SSettingsEditor::Construct( const FArguments& InArgs, const ISettingsEditor
 		]
 	];
 
-	ScrollBox->SetScrollBarRightClickDragAllowed(true);
-
-	FInternationalization::Get().OnCultureChanged().AddSP(this, &SSettingsEditor::HandleCultureChanged);
 	Model->OnSelectionChanged().AddSP(this, &SSettingsEditor::HandleModelSelectionChanged);
-	SettingsContainer->OnCategoryModified().AddSP(this, &SSettingsEditor::HandleSettingsContainerCategoryModified);
 
-	ReloadCategories();
+	// Create category tree later so everything is initialized when callback is triggered
+	CategoryTreeBorder->SetContent(
+		SAssignNew(CategoryTree, SSettingsEditorCategoryTree, SettingsContainer)
+		.InitialSelection(Model->GetSelectedSection())
+		.OnSelectionChanged(this, &SSettingsEditor::OnSectionSelectionChanged)
+	);
+
+	HandleModelSelectionChanged();
+
+	// Update toolbar menu
+	UpdateToolbarMenu();
+	if (ISettingsEditorModule* SettingsEditorModule = FModuleManager::GetModulePtr<ISettingsEditorModule>(TEXT("SettingsEditor")))
+	{
+		SettingsEditorModule->OnSearchTermsChanged().AddSP(this, &SSettingsEditor::OnSearchTermsChanged);
+	}
 }
 
 /* FNotifyHook interface
@@ -352,73 +349,6 @@ TWeakObjectPtr<UObject> SSettingsEditor::GetSelectedSettingsObject() const
 	return nullptr;
 }
 
-TSharedRef<SWidget> SSettingsEditor::MakeCategoryWidget( const ISettingsCategoryRef& Category, const TArray<ISettingsSectionPtr>& InSections )
-{
-	// create section widgets
-	TSharedRef<SVerticalBox> SectionsBox = SNew(SVerticalBox);
-
-	
-	if (InSections.Num() == 0)
-	{
-		return SNullWidget::NullWidget;
-	}
-
-	// list the sections
-	for (const ISettingsSectionPtr& Section : InSections)
-	{
-		SectionsBox->AddSlot()
-			.HAlign(HAlign_Left)
-			.Padding(0.0f, 10.0f, 0.0f, 0.0f)
-			[
-				SNew(SHorizontalBox)
-
-				+ SHorizontalBox::Slot()
-					.AutoWidth()
-					.Padding(0.0f, 0.0f, 4.0f, 0.0f)
-					.VAlign(VAlign_Center)
-					[
-						SNew(SImage)
-							.Image(FAppStyle::Get().GetBrush("TreeArrow_Collapsed_Hovered"))
-							.Visibility(this, &SSettingsEditor::HandleSectionLinkImageVisibility, Section)
-					]
-
-				+ SHorizontalBox::Slot()
-					.FillWidth(1.0f)
-					.VAlign(VAlign_Center)
-					[
-						SNew(SHyperlink)
-							.OnNavigate(this, &SSettingsEditor::HandleSectionLinkNavigate, Section)
-							.Text(Section->GetDisplayName())
-							.ToolTipText(Section->GetDescription())
-					]
-			];
-
-		if (!Model->GetSelectedSection().IsValid())
-		{
-			Model->SelectSection(Section);
-		}
-	}
-
-	// create category widget
-	return SNew(SVerticalBox)
-
-	+ SVerticalBox::Slot()
-		.AutoHeight()
-		[
-			// category title
-			SNew(STextBlock)
-				.Font(FAppStyle::GetFontStyle("SettingsEditor.CatgoryAndSectionFont"))
-				.Text(Category->GetDisplayName())
-		]
-
-	+ SVerticalBox::Slot()
-		.FillHeight(1.0f)
-		[
-			// sections list
-			SectionsBox
-		];
-}
-
 
 void SSettingsEditor::RecordPreferenceChangedAnalytics( ISettingsSectionPtr SelectedSection, const FPropertyChangedEvent& PropertyChangedEvent ) const
 {
@@ -435,109 +365,9 @@ void SSettingsEditor::RecordPreferenceChangedAnalytics( ISettingsSectionPtr Sele
 	}
 }
 
-
-void SSettingsEditor::ReloadCategories()
-{
-	CategoriesBox->ClearChildren();
-
-	CategoriesBox->AddSlot()
-	.AutoHeight()
-	.HAlign(HAlign_Left)
-	.Padding(0.0f, 5.0f, 0.0f, 18.0f)
-	[
-		SNew(SHorizontalBox)
-		+ SHorizontalBox::Slot()
-		.AutoWidth()
-		.Padding(0.0f, 0.0f, 4.0f, 0.0f)
-		.VAlign(VAlign_Center)
-		[
-			SNew(SImage)
-			.Image(FAppStyle::Get().GetBrush("TreeArrow_Collapsed_Hovered"))
-			.Visibility(this, &SSettingsEditor::HandleAllSectionsLinkImageVisibility)
-		]
-		+ SHorizontalBox::Slot()
-		.FillWidth(1.0f)
-		.VAlign(VAlign_Center)
-		[
-			SNew(SHyperlink)
-			.OnNavigate(this, &SSettingsEditor::HandleAllSectionsLinkNavigate)
-			.Text(LOCTEXT("AllPropertiesLink", "All Settings"))
-			.ToolTipText(LOCTEXT("AllPropertiesLink_Tooltip", "Show all settings"))
-		]
-	];
-
-	TArray<ISettingsCategoryPtr> Categories;
-	SettingsContainer->GetCategories(Categories);
-
-	TArray<TSharedPtr<ISettingsSection>> CategorySettingsSections;
-
-	TArray<UObject*> SettingsObjects;
-	for (const ISettingsCategoryPtr& Category : Categories)
-	{
-		CategorySettingsSections.Reset();
-
-		Category->GetSections(CategorySettingsSections);
-
-		// sort the sections alphabetically
-		struct FSectionSortPredicate
-		{
-			FORCEINLINE bool operator()(ISettingsSectionPtr A, ISettingsSectionPtr B) const
-			{
-				if (!A.IsValid() && !B.IsValid())
-				{
-					return false;
-				}
-
-				if (A.IsValid() != B.IsValid())
-				{
-					return B.IsValid();
-				}
-
-				return (A->GetDisplayName().CompareTo(B->GetDisplayName()) < 0);
-			}
-		};
-
-		CategorySettingsSections.Sort(FSectionSortPredicate());
-
-		TSharedRef<SWidget> CategoryWidget = MakeCategoryWidget(Category.ToSharedRef(), CategorySettingsSections);
-		if (CategoryWidget != SNullWidget::NullWidget)
-		{
-			CategoriesBox->AddSlot()
-				.AutoHeight()
-				.Padding(0.0f, 0.0f, 0.0f, 16.0f)
-				[
-					MoveTemp(CategoryWidget)
-				];
-		}
-
-		for(TSharedPtr<ISettingsSection>& Section : CategorySettingsSections)
-		{
-			UObject* SettingsObject = Section->GetSettingsObject().Get();
-			if(SettingsObject)
-			{
-				SettingsObjects.Add(SettingsObject);
-			}
-		}
-	}
-
-	Model->SelectSection(nullptr);
-	SettingsView->SetObjects(SettingsObjects);
-
-
-}
-
-
-
-
 /* SSettingsEditor callbacks
  *****************************************************************************/
 
-
-void SSettingsEditor::HandleCultureChanged()
-{
-	ReloadCategories();
-}
-
 void SSettingsEditor::HandleModelSelectionChanged()
 {
 	// This callback can trigger on unregister during shutdown, simply return in this case
@@ -546,14 +376,10 @@ void SSettingsEditor::HandleModelSelectionChanged()
 		return;
 	}
 
-	ISettingsSectionPtr SelectedSection = Model->GetSelectedSection();
-
-	if (SelectedSection.IsValid())
+	if (const ISettingsSectionPtr SelectedSection = Model->GetSelectedSection())
 	{
-		TSharedPtr<SWidget> CustomWidget = SelectedSection->GetCustomWidget().Pin();
-
 		// show settings widget
-		if (CustomWidget.IsValid())
+		if (const TSharedPtr<SWidget> CustomWidget = SelectedSection->GetCustomWidget().Pin())
 		{
 			CustomWidgetSlot->AttachWidget( CustomWidget.ToSharedRef() );
 		}
@@ -562,61 +388,41 @@ void SSettingsEditor::HandleModelSelectionChanged()
 			CustomWidgetSlot->AttachWidget( SNullWidget::NullWidget );
 		}
 
-		// focus settings widget
-		TSharedPtr<SWidget> FocusWidget;
-
-		if (CustomWidget.IsValid())
-		{
-			FocusWidget = CustomWidget;
-		}
-		else
-		{
-			FocusWidget = SettingsView;
-		}
-
-		FWidgetPath FocusWidgetPath;
-
-		if (FSlateApplication::Get().GeneratePathToWidgetUnchecked(FocusWidget.ToSharedRef(), FocusWidgetPath))
-		{
-			FSlateApplication::Get().SetKeyboardFocus(FocusWidgetPath, EFocusCause::SetDirectly);
-		}
-
 		bShowingAllSettings = false;
+
+		// Sync tree view with selection
+		CategoryTree->SetSelection(SelectedSection);
 	}
 	else
 	{
 		bShowingAllSettings = true;
 		CustomWidgetSlot->AttachWidget( SNullWidget::NullWidget );
+
+		// Sync tree view with selection
+		CategoryTree->SetSelection(/** Section */nullptr);
 	}
 
 	// clear the global search terms when selecting a specific category
-	SettingsView->ClearSearch();
-}
+	SearchKey.Reset();
+	bAutoDetectKeybind = false;
+	SetSearchFilter(FText::GetEmpty());
 
-void SSettingsEditor::HandleSectionLinkNavigate( ISettingsSectionPtr Section )
-{
-	Model->SelectSection(Section);
-}
-
-void SSettingsEditor::HandleAllSectionsLinkNavigate()
-{
-	Model->SelectSection(nullptr);
-	SettingsView->RefreshRootObjectVisibility();
-}
-
-EVisibility SSettingsEditor::HandleAllSectionsLinkImageVisibility() const
-{
-	return bShowingAllSettings ? EVisibility::Visible : EVisibility::Hidden;
-}
-
-EVisibility SSettingsEditor::HandleSectionLinkImageVisibility( ISettingsSectionPtr Section ) const
-{
-	if (Model->GetSelectedSection() == Section)
+	// Set all settings object to preview them,
+	// they will get filtered base on the active section selection in the customization
 	{
-		return EVisibility::Visible;
+		TArray<UObject*> SettingsObjects;
+		CategoryTree->ForEachSection([&SettingsObjects](const TSharedPtr<ISettingsSection>& Section)
+			{
+				if (UObject* SettingObject = Section->GetSettingsObject().Get())
+				{
+					SettingsObjects.Add(SettingObject);
+				}
+				
+				return true; // Continue
+			}
+		);
+		SettingsView->SetObjects(SettingsObjects);
 	}
-
-	return EVisibility::Hidden;
 }
 
 EVisibility SSettingsEditor::HandleSettingsBoxVisibility() const
@@ -631,24 +437,6 @@ EVisibility SSettingsEditor::HandleSettingsBoxVisibility() const
 	return EVisibility::Hidden;
 }
 
-void SSettingsEditor::HandleSettingsContainerCategoryModified( const FName& CategoryName )
-{
-	if ( !bIsActiveTimerRegistered )
-	{
-		bIsActiveTimerRegistered = true;
-		FTSTicker::GetCoreTicker().AddTicker(FTickerDelegate::CreateSP(this, &SSettingsEditor::UpdateCategoriesCallback));
-	}
-}
-
-bool SSettingsEditor::UpdateCategoriesCallback(float InDeltaTime)
-{
-	bIsActiveTimerRegistered = false;
-
-	ReloadCategories();
-
-	return false;
-}
-
 void ShowNotification(const FText& Text, SNotificationItem::ECompletionState CompletionState)
 {
 	FNotificationInfo Notification(Text);
@@ -801,31 +589,42 @@ FReply SSettingsEditor::HandleImportButtonClicked()
 	return FReply::Handled();
 }
 
-FText SSettingsEditor::GetSearchToolTipText() const
+FText SSettingsEditor::GetSearchHintText() const
 {
-	FText ToolTipText = LOCTEXT("FilterSearchTooltip", "Type here to search...");
-
-	if (const ISettingsEditorModule* SettingsEditorModule = FModuleManager::GetModulePtr<ISettingsEditorModule>(TEXT("SettingsEditor")))
+	// Custom hint for when we are listening for key presses
+	static const FText KeyBindSearch = LOCTEXT("KeyBindSearch", "(Press any key(s) to search)");
+	if (bAutoDetectKeybind)
 	{
-		if (SettingsContainer)
-		{
-			if (const TMap<FString, TSet<FString>>* SearchTerms = SettingsEditorModule->FindSearchTerms(SettingsContainer->GetName()))
-			{
-				TArray<FString> Keys;
-				SearchTerms->GenerateKeyArray(Keys);
-				ToolTipText = FText::Format(INVTEXT("{0} ({1})"), ToolTipText, FText::FromString(FString::Join(Keys, TEXT(", "))));
-			}
-		}
+		return KeyBindSearch;
 	}
 
-	return ToolTipText;
+	FText SectionText;
+	if (const TSharedPtr<ISettingsSection> SelectedSection = Model->GetSelectedSection())
+	{
+		SectionText = SelectedSection->GetDisplayName();
+	}
+	else
+	{
+		SectionText = LOCTEXT("SettingsAllSection", "all settings");
+	}
+	
+	if (SearchKey.IsEmpty())
+	{
+		static const FText HintFormat = LOCTEXT("SearchHintText", "Search in {0}...");
+		return FText::Format(HintFormat, SectionText);
+	}
+	else
+	{
+		static const FText AdvancedHintFormat = LOCTEXT("SearchAdvancedHintText", "Search in {0} for {1}...");
+		return FText::Format(AdvancedHintFormat, SectionText, FText::FromString(SearchKey.ToLower()));
+	}
 }
 
 void SSettingsEditor::OnSearchTextChanged(const FText& InText)
 {
 	if (SettingsView.IsValid())
 	{
-		SettingsView->SetSearchText(InText);
+		SetSearchFilter(InText);
 	}
 }
 
@@ -833,7 +632,7 @@ void SSettingsEditor::OnSearchTextCommitted(const FText& InText, ETextCommit::Ty
 {
 	if (SettingsView.IsValid() && InType == ETextCommit::Type::OnCleared)
 	{
-		SettingsView->ClearSearch();
+		SetSearchFilter(FText::GetEmpty());
 	}
 }
 
@@ -845,41 +644,83 @@ void SSettingsEditor::OnAssetSearchSuggestionFilter(const FText& InText, TArray<
 		return;
 	}
 
-	const TMap<FString, TSet<FString>>* SearchTerms = SettingsEditorModule->FindSearchTerms(SettingsContainer->GetName());
+	const TMap<FString, ISettingsEditorModule::FSearchTerm>* SearchTerms = SettingsEditorModule->FindSearchTerms(SettingsContainer->GetName());
 	if (!SearchTerms)
 	{
 		return;
 	}
 
 	static const FText SettingsCategory = LOCTEXT("SettingsEditor", "Settings");
-	static const FTextFormat SearchFormat = INVTEXT("{0}={1}");
-	const FString ActiveSearch = InText.ToString();
-	for (const TPair<FString, TSet<FString>>& SearchTerm : *SearchTerms)
+
+	// Do not append search key to the suggestion again
+	FString ActiveSearch;
+	if (!SearchKey.IsEmpty())
 	{
-		if (ActiveSearch.Contains(SearchTerm.Key) || SearchTerm.Key.Contains(ActiveSearch))
+		ActiveSearch = FText::FormatOrdered(TokenCombineFormat, FText::FromString(SearchKey), InText).ToString();
+	}
+	else
+	{
+		ActiveSearch = InText.ToString();
+	}
+
+	const FName SelectedSectionName = GetSelectedSectionName();
+	for (const TPair<FString, ISettingsEditorModule::FSearchTerm>& SearchTerm : *SearchTerms)
+	{
+		if (!SearchTerm.Value.PassesSectionFilter(SelectedSectionName))
+		{
+			continue;
+		}
+
+		// Suggest search term first, then suggest search term + search value 
+		// eg: typing "cv" will show "cvar", and typing "cvar" will show "cvar R.VSync"
+		const bool bActiveSearchContainsSearchTermKey = ActiveSearch.Contains(SearchTerm.Key);
+		const bool bSearchTermContainsActiveSearch = SearchTerm.Key.Contains(ActiveSearch);
+		if (bActiveSearchContainsSearchTermKey || bSearchTermContainsActiveSearch)
 		{
-			FString SearchSuggestion = SearchTerm.Key.ToLower() + TEXT("=");
-			FText SearchSuggestionLabel = FText::FormatOrdered(SearchFormat, FText::FromString(SearchTerm.Key), FText::GetEmpty());
-			OutSuggestion.Add(FAssetSearchBoxSuggestion(SearchSuggestion, SearchSuggestionLabel, SettingsCategory));
+			FString SearchBase = !SearchKey.IsEmpty() ? TEXT("") : SearchTerm.Key.ToLower();
+			FString SearchSuggestion = SearchBase;
+			FText SearchSuggestionLabel;
+			if (!SearchBase.IsEmpty() && !bActiveSearchContainsSearchTermKey)
+			{
+				SearchSuggestionLabel = FText::FromString(SearchBase);
+				OutSuggestion.Add(FAssetSearchBoxSuggestion(SearchSuggestion, SearchSuggestionLabel, SettingsCategory));
+			}
+
 			FString SearchCopy = ActiveSearch;
-			SearchCopy.RemoveFromStart(SearchSuggestion);
+			SearchCopy.RemoveFromStart(SearchTerm.Key.ToLower() + TEXT(" "));
 
-			for (const FString& SearchValue : SearchTerm.Value)
+			for (const FString& SearchValue : SearchTerm.Value.Values)
 			{
-				if (SearchCopy.IsEmpty() || SearchValue.StartsWith(SearchCopy))
+				// Allow partial search (eg: "collect" will list "gc.collecteveryframe")
+				if (bSearchTermContainsActiveSearch || SearchCopy.IsEmpty() || SearchValue.Contains(SearchCopy))
 				{
-					SearchSuggestion = SearchTerm.Key.ToLower() + TEXT("=") + SearchValue;
-					SearchSuggestionLabel = FText::FormatOrdered(SearchFormat, FText::FromString(SearchTerm.Key), FText::FromString(SearchValue));
+					SearchSuggestion = SearchBase + TEXT(" \"") + SearchValue + TEXT("\"");
+					SearchSuggestionLabel = FText::FromString(SearchSuggestion);
 					OutSuggestion.Add(FAssetSearchBoxSuggestion(SearchSuggestion, SearchSuggestionLabel, SettingsCategory));
 				}
 			}
 		}
 	}
+
+	// Clear suggestions when we already have it filled
+	if (OutSuggestion.Num() == 1 && OutSuggestion[0].SuggestionString == InText.ToString())
+	{
+		OutSuggestion.Reset();
+	}
 }
 
 FText SSettingsEditor::OnAssetSearchSuggestionChosen(const FText& InText, const FString& InSuggestion)
 {
 	FString CurrentSearch = InText.ToString();
+
+	// Trim the start, eg: "keybind ctrl" becomes "ctrl"
+	// This is needed because the "..." used for strict search do not match with the current search (eg: keybind "ctrl" vs keybind ctrl)
+	int32 SpaceIndex;
+	if (CurrentSearch.FindChar(' ', SpaceIndex))
+	{
+		CurrentSearch = CurrentSearch.RightChop(SpaceIndex + 1);
+	}
+
 	if (!InSuggestion.Contains(CurrentSearch))
 	{
 		for (const auto& Char : InSuggestion)
@@ -929,85 +770,227 @@ void SSettingsEditor::FillSettingsFiltersSection(UToolMenu* InMenu)
 		return;
 	}
 
-	const TMap<FString, TSet<FString>>* SearchTerms = SettingsEditorModule->FindSearchTerms(SettingsEditor->SettingsContainer->GetName());
+	const TMap<FString, ISettingsEditorModule::FSearchTerm>* SearchTerms = SettingsEditorModule->FindSearchTerms(SettingsEditor->SettingsContainer->GetName());
 	if (!SearchTerms)
 	{
 		return;
 	}
 
 	FToolMenuSection& FiltersSection = InMenu->FindOrAddSection("Filters", LOCTEXT("FiltersLabel", "Filters"));
-	for (const TPair<FString, TSet<FString>>& SearchTerm : *SearchTerms)
+	for (const TPair<FString, ISettingsEditorModule::FSearchTerm>& SearchTerm : *SearchTerms)
 	{
+		FToolMenuEntry MenuEntry = CreateFiltersCheckboxEntry(SettingsEditorModule, Context, SearchTerm.Key, SearchTerm.Value.Values);
+
+		// Customize icon
+		if (!SearchTerm.Value.IconStyleSetName.IsNone() && !SearchTerm.Value.IconStyleName.IsNone())
+		{
+			MenuEntry.Icon = FSlateIcon(SearchTerm.Value.IconStyleSetName, SearchTerm.Value.IconStyleName);
+		}
+		else
+		{
+			MenuEntry.Icon = FSlateIcon(FAppStyle::GetAppStyleSetName(), "Icons.Filter");
+		}
+
+		// Customize label
+		if (!SearchTerm.Value.Label.IsEmpty())
+		{
+			MenuEntry.Label = SearchTerm.Value.Label;
+		}
+
 		FiltersSection
-			.AddEntry(CreateFiltersCheckboxEntry(Context, SearchTerm.Key, SearchTerm.Value))
+			.AddEntry(MenuEntry)
 			.SetShowInToolbarTopLevel(true);
 	}
 }
 
-FToolMenuEntry SSettingsEditor::CreateFiltersCheckboxEntry(const USettingsEditorMenuContext* InContext, const FString& InKey, const TSet<FString>& InValues)
+FToolMenuEntry SSettingsEditor::CreateFiltersCheckboxEntry(const ISettingsEditorModule* InSettingsEditorModule, const USettingsEditorMenuContext* InContext, const FString& InKey, const TSet<FString>& InValues)
 {
 	const TSharedPtr<SSettingsEditor> SettingsEditor = InContext->GetSettingsEditor();
-	checkf(SettingsEditor.IsValid(), TEXT("SettingsEditor must be valid"))
+	checkf(SettingsEditor.IsValid(), TEXT("SettingsEditor must be valid"));
+	checkf(InSettingsEditorModule != nullptr, TEXT("SettingsEditorModule must be valid"));
 
 	FUIAction FilterAction;
 	FilterAction.ExecuteAction = FExecuteAction::CreateSP(SettingsEditor.Get(), &SSettingsEditor::ToggleSearchFilterAction, InKey);
-	FilterAction.GetActionCheckState = FGetActionCheckState::CreateSP(SettingsEditor.Get(), &SSettingsEditor::IsSearchFilterActive, InKey, /** Equals */false);
+	FilterAction.GetActionCheckState = FGetActionCheckState::CreateSP(SettingsEditor.Get(), &SSettingsEditor::IsSearchFilterActive, InKey);
+	FilterAction.IsActionVisibleDelegate = FIsActionButtonVisible::CreateSP(SettingsEditor.Get(), &SSettingsEditor::IsSearchFilterVisible, InKey);
 
-	FToolMenuEntry FilterEntry = FToolMenuEntry::InitSubMenu(
+	return FToolMenuEntry::InitToolBarButton(
 		FName(InKey),
-		FText::FromString(InKey),
-		LOCTEXT("SearchFilterToolTip", "Toggles this search filter on/off"),
-		FNewToolMenuDelegate::CreateLambda([SettingsEditor, InKey, InValues](UToolMenu* InSubmenu)
-		{
-			FToolMenuSection& FiltersSubSection = InSubmenu->FindOrAddSection(FName(InKey), FText::FromString(InKey));
-			for (const FString& Value : InValues)
-			{
-				const FString SearchTerm = InKey + TEXT("=") + Value;
-				FiltersSubSection.AddMenuEntry(
-					NAME_None,
-					FText::FromString(Value.TrimChar('"')),
-					LOCTEXT("SetFilterToolTip", "Apply search filter with this value"),
-					FSlateIcon(),
-					FUIAction(
-						FExecuteAction::CreateSP(SettingsEditor.Get(), &SSettingsEditor::SetSearchFilter, SearchTerm),
-						FCanExecuteAction(),
-						FGetActionCheckState::CreateSP(SettingsEditor.Get(), &SSettingsEditor::IsSearchFilterActive, SearchTerm, /** Equals */true)
-					),
-					EUserInterfaceActionType::RadioButton
-				);
-			}
-		}),
 		FilterAction,
+		FText::FromString(InKey),
+		TAttribute<FText>::CreateSP(SettingsEditor.Get(), &SSettingsEditor::GetSearchFilterTooltip, InSettingsEditorModule, InKey),
+		TAttribute<FSlateIcon>(),
 		EUserInterfaceActionType::ToggleButton
 	);
-
-	FilterEntry.Icon = FSlateIcon(FAppStyle::GetAppStyleSetName(), "Icons.Filter");
-
-	return FilterEntry;
 }
 
 void SSettingsEditor::ToggleSearchFilterAction(FString InFilter)
 {
-	if (IsSearchFilterActive(InFilter, /** Equals */false) == ECheckBoxState::Unchecked)
+	if (IsSearchFilterActive(InFilter) == ECheckBoxState::Unchecked)
 	{
-		SearchBox->SetText(FText::FromString(InFilter + TEXT("=")));
+		SearchKey = InFilter;
+
+		// Custom use case to detect which keys have been pressed by user and perform a search on those
+		static const FText KeybindKey = NSLOCTEXT("InputBindingEditor", "SearchKeyKeybind", "KeyBind:");
+		if (SearchKey == KeybindKey.ToString())
+		{
+			bAutoDetectKeybind = true;
+		}
 	}
 	else
 	{
-		SearchBox->SetText(FText::GetEmpty());
+		bAutoDetectKeybind = false;
+		SearchKey.Reset();
 	}
+
+	SetSearchFilter(FText::GetEmpty());
+	FSlateApplication::Get().SetKeyboardFocus(SearchBox, EFocusCause::SetDirectly);
 }
 
-ECheckBoxState SSettingsEditor::IsSearchFilterActive(FString InFilter, bool bInStrictCompare)
+ECheckBoxState SSettingsEditor::IsSearchFilterActive(FString InFilter)
 {
-	const FString ActiveSearch = SearchBox->GetText().ToString();
-	const bool bSearchFilterActive = bInStrictCompare ? ActiveSearch.Equals(InFilter) : ActiveSearch.StartsWith(InFilter);
+	const bool bSearchFilterActive = SearchKey.Equals(InFilter, ESearchCase::CaseSensitive);
 	return bSearchFilterActive ? ECheckBoxState::Checked : ECheckBoxState::Unchecked;
 }
 
-void SSettingsEditor::SetSearchFilter(FString InSearchTerm)
+void SSettingsEditor::SetSearchFilter(const FText& InSearch, bool bInStrict)
+{
+	FText SearchText = InSearch;
+	if (bInStrict)
+	{
+		SearchText = FText::FormatOrdered(TokenDelimiter, SearchText);
+	}
+
+	if (!SearchKey.IsEmpty())
+	{
+		SearchText = FText::FormatOrdered(TokenCombineFormat, FText::FromString(SearchKey), SearchText);
+	}
+
+	SearchBox->SetText(InSearch);
+	SettingsView->SetSearchText(SearchText);
+}
+
+bool SSettingsEditor::IsSearchFilterVisible(FString InSearchTerm) const
+{
+	ISettingsEditorModule* SettingsEditorModule = FModuleManager::GetModulePtr<ISettingsEditorModule>(TEXT("SettingsEditor"));
+	if (!SettingsEditorModule || !SettingsContainer || !Model)
+	{
+		return false;
+	}
+
+	const ISettingsEditorModule::FSearchTerm* SearchTerm = SettingsEditorModule->FindSearchTerm(SettingsContainer->GetName(), InSearchTerm);
+	if (!SearchTerm || !SearchTerm->bShowFilter)
+	{
+		return false;
+	}
+
+	const FName SectionName = GetSelectedSectionName();
+	return SearchTerm->PassesSectionFilter(SectionName);
+}
+
+FName SSettingsEditor::GetSelectedSectionName() const
+{
+	FName SectionName = NAME_None;
+	if (!Model)
+	{
+		return SectionName;
+	}
+	const TSharedPtr<ISettingsSection>& SelectedSection = Model->GetSelectedSection();
+	if (!SelectedSection)
+	{
+		return SectionName;
+	}
+
+	SectionName = SelectedSection->GetName();
+	return SectionName;
+}
+
+FText SSettingsEditor::GetSearchFilterTooltip(const ISettingsEditorModule* InSettingsEditorModule, FString InKey) const
 {
-	SearchBox->SetText(FText::FromString(InSearchTerm));
+	static const FTextFormat TooltipFormat = INVTEXT("{0}\n{1}");
+	static const FText StatusTooltipText = LOCTEXT("DisabledSearchFilterToolTip", "Filter disabled since it is not relevant to the selected section");
+	static const FText DefaultTooltipText = LOCTEXT("SearchFilterToolTip", "Toggles search filter to only show entries relevant to the selected context");
+
+	FText TooltipText = DefaultTooltipText;
+	if (SettingsContainer && InSettingsEditorModule != nullptr)
+	{
+		if (const ISettingsEditorModule::FSearchTerm* SearchTerms = InSettingsEditorModule->FindSearchTerm(SettingsContainer->GetName(), InKey))
+		{
+			if (!SearchTerms->Tooltip.IsEmpty())
+			{
+				TooltipText = SearchTerms->Tooltip;
+			}
+		}
+	}
+
+	const bool bIsFilterEnabled = IsSearchFilterVisible(InKey);
+	return bIsFilterEnabled ? TooltipText : FText::Format(TooltipFormat, TooltipText, StatusTooltipText);
+}
+
+void SSettingsEditor::OnSectionSelectionChanged(TSharedPtr<ISettingsSection> InSection)
+{
+	Model->SelectSection(InSection);
+}
+
+FReply SSettingsEditor::OnSearchKeyDownHandler(const FGeometry& InGeometry, const FKeyEvent& InEvent)
+{
+	if (bAutoDetectKeybind)
+	{
+		const FInputChord InputChord(InEvent.GetKey(), InEvent.IsShiftDown(), InEvent.IsControlDown(), InEvent.IsAltDown(), InEvent.IsCommandDown());
+
+		// Once a non modifier key is pressed, stop listening
+		if (!InEvent.GetKey().IsModifierKey())
+		{
+			// Clear focus to avoid key to be processed like backspace
+			FSlateApplication::Get().ClearKeyboardFocus(EFocusCause::Cleared);
+
+			// Refocus search box
+			FTSTicker::GetCoreTicker().AddTicker(FTickerDelegate::CreateSPLambda(this, [this](float)->bool
+				{
+					if (FSlateApplication::IsInitialized())
+					{
+						FSlateApplication::Get().SetKeyboardFocus(SearchBox, EFocusCause::SetDirectly);
+					}
+					return false; // Stop
+				})
+			);
+		}
+
+		// Perform a strict search here
+		constexpr bool bStrictSearch = true;
+		SetSearchFilter(InputChord.GetInputText(), bStrictSearch);
+		return FReply::Handled();
+	}
+
+	return FReply::Unhandled();
+}
+
+void SSettingsEditor::OnSearchTermsChanged(FName InContainerName, const FString& InKey)
+{
+	if (MenuToolbarUpdateHandle.IsValid())
+	{
+		return;
+	}
+
+	MenuToolbarUpdateHandle = FTSTicker::GetCoreTicker().AddTicker(
+		FTickerDelegate::CreateSPLambda(this, [this](float)
+		{
+			UpdateToolbarMenu();
+			MenuToolbarUpdateHandle.Reset();
+			return false;
+		})
+	);
+}
+
+void SSettingsEditor::UpdateToolbarMenu()
+{
+	if (UToolMenus* ToolMenus = UToolMenus::Get())
+	{
+		USettingsEditorMenuContext* MenuContext = NewObject<USettingsEditorMenuContext>();
+		MenuContext->SettingsEditorWeak = SharedThis(this);
+		const TSharedRef<SWidget> MenuWidget = ToolMenus->GenerateWidget(SettingsEditorToolbarName, FToolMenuContext(MenuContext));
+		MenuToolbarBox->SetContent(MenuWidget);
+	}
 }
 
 bool SSettingsEditor::HandleSettingsViewEnabled() const
diff --git a/Engine/Source/Developer/SettingsEditor/Private/Widgets/SSettingsEditor.h b/Engine/Source/Developer/SettingsEditor/Private/Widgets/SSettingsEditor.h
--- a/Engine/Source/Developer/SettingsEditor/Private/Widgets/SSettingsEditor.h
+++ b/Engine/Source/Developer/SettingsEditor/Private/Widgets/SSettingsEditor.h
@@ -3,6 +3,7 @@
 #pragma once
 
 #include "CoreMinimal.h"
+#include "Containers/Ticker.h"
 #include "ISettingsCategory.h"
 #include "ISettingsContainer.h"
 #include "ISettingsEditorModel.h"
@@ -16,8 +17,9 @@
 
 class IDetailsView;
 class SAssetSearchBox;
+class SBox;
 class SScrollBox;
-class SVerticalBox;
+class SSettingsEditorCategoryTree;
 class USettingsEditorMenuContext;
 class UToolMenu;
 enum class ECheckBoxState : uint8;
@@ -33,6 +35,8 @@ class SSettingsEditor
 	, public FNotifyHook
 {
 public:
+	/** Menu toolbar Identifier */
+	static FLazyName SettingsEditorToolbarName;
 
 	SLATE_BEGIN_ARGS(SSettingsEditor) { }
 		SLATE_EVENT( FSimpleDelegate, OnApplicationRestartRequired )
@@ -69,14 +73,6 @@ protected:
 	 */
 	TWeakObjectPtr<UObject> GetSelectedSettingsObject() const;
 
-	/**
-	 * Creates a widget for the given settings category.
-	 *
-	 * @param Category The category to create the widget for.
-	 * @return The created widget.
-	 */
-	TSharedRef<SWidget> MakeCategoryWidget( const ISettingsCategoryRef& Category, const TArray<ISettingsSectionPtr>& InSections );
-
 
 	/**
 	 * Reports a preference changed event to the analytics system.
@@ -86,68 +82,55 @@ protected:
 	 */
 	void RecordPreferenceChangedAnalytics( ISettingsSectionPtr SelectedSection, const FPropertyChangedEvent& PropertyChangedEvent ) const;
 
-	/** Reloads the settings categories. */
-	void ReloadCategories();
-
-
 private:
-	static void FillSettingsFiltersSection(UToolMenu* InMenu);
-	static FToolMenuEntry CreateFiltersCheckboxEntry(const USettingsEditorMenuContext* InContext, const FString& InKey, const TSet<FString>& InValues);
+	static const FTextFormat TokenCombineFormat;
+	static const FTextFormat TokenDelimiter;
 
-	/** Callback for when the user's culture has changed. */
-	void HandleCultureChanged();
+	static void FillSettingsFiltersSection(UToolMenu* InMenu);
+	static FToolMenuEntry CreateFiltersCheckboxEntry(const ISettingsEditorModule* InSettingsEditorModule, const USettingsEditorMenuContext* InContext, const FString& InKey, const TSet<FString>& InValues);
 
 	/** Callback for changing the selected settings section. */
 	void HandleModelSelectionChanged();
 
-	/** Callback for navigating a settings section link. */
-	void HandleSectionLinkNavigate( ISettingsSectionPtr Section );
-
-	/** Callback for navigating the all settings link. */
-	void HandleAllSectionsLinkNavigate();
-
-	/** Callback for getting the visibility of a section link image. */
-	EVisibility HandleSectionLinkImageVisibility( ISettingsSectionPtr Section ) const;
-
-	/** Callback for getting the visibility of a the all sections link image. */
-	EVisibility HandleAllSectionsLinkImageVisibility() const;
-
 	/** Callback for determining the visibility of the settings box. */
 	EVisibility HandleSettingsBoxVisibility() const;
 
-	/** Callback for the modification of categories in the settings container. */
-	void HandleSettingsContainerCategoryModified( const FName& CategoryName );
-
 	/** Callback for checking whether the settings view is enabled. */
 	bool HandleSettingsViewEnabled() const;
 
 	/** Callback for determining the visibility of the settings view. */
 	EVisibility HandleSettingsViewVisibility() const;
 
-	/** Callback when the timer has been ticked to refresh the categories latently. */
-	bool UpdateCategoriesCallback(float InDeltaTime);
-
 	/** Callback for exporting all sections as a single ini. */
 	FReply HandleExportButtonClicked();
 
 	/** Callback for importing all sections from a single ini. */
 	FReply HandleImportButtonClicked();
 
-	FText GetSearchToolTipText() const;
+	FText GetSearchHintText() const;
 	void OnSearchTextChanged(const FText& InText);
 	void OnSearchTextCommitted(const FText& InText, ETextCommit::Type InType);
 	void OnAssetSearchSuggestionFilter(const FText& InText, TArray<FAssetSearchBoxSuggestion>& OutSuggestion, FText& OutText);
 	FText OnAssetSearchSuggestionChosen(const FText& InText, const FString& InSuggestion);
 
 	void ToggleSearchFilterAction(FString InFilter);
-	ECheckBoxState IsSearchFilterActive(FString InFilter, bool bInStrictCompare);
-	void SetSearchFilter(FString InSearchTerm);
+	ECheckBoxState IsSearchFilterActive(FString InFilter);
+	void SetSearchFilter(const FText& InSearch, bool bInStrict = false);
+	bool IsSearchFilterVisible(FString InSearchTerm) const;
+	FName GetSelectedSectionName() const;
+	FText GetSearchFilterTooltip(const ISettingsEditorModule* InSettingsEditorModule, FString InKey) const;
+
+	void OnSectionSelectionChanged(TSharedPtr<ISettingsSection> InSection);
+	FReply OnSearchKeyDownHandler(const FGeometry& InGeometry, const FKeyEvent& InEvent);
+	void OnSearchTermsChanged(FName InContainerName, const FString& InKey);
+
+	void UpdateToolbarMenu();
 
-	/** Holds the vertical box for settings categories. */
-	TSharedPtr<SVerticalBox> CategoriesBox;
+	/** Holds the settings category tree. */
+	TSharedPtr<SSettingsEditorCategoryTree> CategoryTree;
 
 	/** Holds the overlay slot for custom widgets. */
-	SOverlay::FOverlaySlot* CustomWidgetSlot;
+	SOverlay::FOverlaySlot* CustomWidgetSlot = nullptr;
 
 	/** Holds a pointer to the view model. */
 	ISettingsEditorModelPtr Model;
@@ -160,13 +143,22 @@ private:
 
 	/** Delegate called when this settings editor requests that the user be notified that the application needs to be restarted for some setting changes to take effect */
 	FSimpleDelegate OnApplicationRestartRequiredDelegate;
-
-	/** Is the active timer registered to refresh categories after the settings changed. */
-	bool bIsActiveTimerRegistered;
 	
 	/** Are we showing all settings at once */
-	bool bShowingAllSettings;
+	bool bShowingAllSettings = false;
 
 	/** Actual search box with suggestions */
 	TSharedPtr<SAssetSearchBox> SearchBox;
+	
+	/** Used to detect when to listen for key down events */
+	bool bAutoDetectKeybind = false;
+	
+	/** Box for the horizontal menu toolbar */
+	TSharedPtr<SBox> MenuToolbarBox = nullptr;
+	
+	/** Handle for when an update is planned */
+	FTSTicker::FDelegateHandle MenuToolbarUpdateHandle;
+
+	/** Used to store active search key */
+	FString SearchKey;
 };
diff --git a/Engine/Source/Developer/SettingsEditor/Private/Widgets/SSettingsEditorCategoryItem.cpp b/Engine/Source/Developer/SettingsEditor/Private/Widgets/SSettingsEditorCategoryItem.cpp
new file mode 100644
--- /dev/null
+++ b/Engine/Source/Developer/SettingsEditor/Private/Widgets/SSettingsEditorCategoryItem.cpp
@@ -0,0 +1,89 @@
+// Copyright Epic Games, Inc. All Rights Reserved.
+
+#include "Widgets/SSettingsEditorCategoryItem.h"
+
+#include "Internationalization/Internationalization.h"
+#include "ISettingsCategory.h"
+#include "ISettingsSection.h"
+#include "Styling/AppStyle.h"
+#include "Widgets/Layout/SBorder.h"
+#include "Widgets/SBoxPanel.h"
+#include "Widgets/Text/STextBlock.h"
+#include "Widgets/Views/SExpanderArrow.h"
+
+#define LOCTEXT_NAMESPACE "SettingsEditorCategoryItem"
+
+using namespace UE::SettingsEditor::Private;
+
+void SSettingsEditorCategoryItem::Construct(const FArguments& InArgs, const TSharedRef<STableViewBase>& InOwner, const TSharedRef<FTreeItem>& InItem)
+{
+	ItemWeak = InItem;
+
+	// Figure out which style to use
+	const auto bIsCategory = InItem->Category.IsValid() && !InItem->Section.IsValid();
+	const auto bIsAllCategory = !InItem->Category.IsValid() && !InItem->Section.IsValid();
+
+	FText Label;
+	FText ToolTipText;
+	FName BorderPaddingName;
+	if (bIsAllCategory)
+	{
+		Label = LOCTEXT("AllPropertiesLink", "All Settings");
+		ToolTipText = LOCTEXT("AllPropertiesLink_Tooltip", "Show all settings");
+		BorderPaddingName = TEXT("CategoryTreeItem.Root.AllBackgroundPadding");
+	}
+	else if (bIsCategory)
+	{
+		Label = InItem->Category->GetDisplayName();
+		ToolTipText = InItem->Category->GetDescription();
+		BorderPaddingName = TEXT("CategoryTreeItem.Root.BackgroundPadding");
+	}
+	else
+	{
+		Label = InItem->Section->GetDisplayName();
+		ToolTipText = InItem->Section->GetDescription();
+		BorderPaddingName = TEXT("CategoryTreeItem.BackgroundPadding");
+	}
+
+	this->ChildSlot
+	[
+		SNew(SHorizontalBox)
+		+ SHorizontalBox::Slot()
+		.AutoWidth()
+		.HAlign(HAlign_Right)
+		.VAlign(VAlign_Fill)
+		[
+			SNew(SExpanderArrow, SharedThis(this))
+			.IndentAmount(0)
+		]
+		+ SHorizontalBox::Slot()
+		.FillWidth(1)
+		[
+			SNew(SBorder)
+			.Padding(FAppStyle::Get().GetMargin(BorderPaddingName))
+			.BorderImage(FStyleDefaults::GetNoBrush())
+			[
+				SNew(SHorizontalBox)
+				// Category name
+				+SHorizontalBox::Slot()
+				.VAlign(VAlign_Center)
+				.Padding(FAppStyle::Get().GetFloat("CategoryTreeItem.PaddingAmount"))
+				[
+					SNew(STextBlock)
+					.Text(Label)
+					.ToolTipText(ToolTipText)
+					.TextStyle(FAppStyle::Get(), (bIsCategory || bIsAllCategory) ? "CategoryTreeItem.Root.Text" : "CategoryTreeItem.Text")
+				]
+			]
+		]
+	];
+
+	STableRow<TSharedPtr<FTreeItem>>::ConstructInternal(
+		STableRow::FArguments()
+			.Style(&FAppStyle::Get(), "SimpleTableView.Row")
+			.ShowSelection(!bIsCategory),
+		InOwner
+	);
+}
+
+#undef LOCTEXT_NAMESPACE
diff --git a/Engine/Source/Developer/SettingsEditor/Private/Widgets/SSettingsEditorCategoryItem.h b/Engine/Source/Developer/SettingsEditor/Private/Widgets/SSettingsEditorCategoryItem.h
new file mode 100644
--- /dev/null
+++ b/Engine/Source/Developer/SettingsEditor/Private/Widgets/SSettingsEditorCategoryItem.h
@@ -0,0 +1,37 @@
+// Copyright Epic Games, Inc. All Rights Reserved.
+
+#pragma once
+
+#include "Widgets/Views/ITableRow.h"
+#include "Widgets/Views/SListView.h"
+#include "Widgets/Views/STableRow.h"
+
+class ISettingsSection;
+class ISettingsCategory;
+class STableViewBase;
+
+namespace UE::SettingsEditor::Private
+{
+	struct FTreeItem
+	{
+		TSharedPtr<ISettingsCategory> Category;
+		TSharedPtr<ISettingsSection> Section;
+		TArray<TSharedPtr<FTreeItem>> Children;
+	};
+}
+
+/**
+ * Widget that represents a single item in SSettingsEditorCategoryTree
+ */
+class SSettingsEditorCategoryItem : public STableRow<TSharedPtr<UE::SettingsEditor::Private::FTreeItem>>
+{
+public:
+	SLATE_BEGIN_ARGS(SSettingsEditorCategoryItem)
+	{}
+	SLATE_END_ARGS()
+
+	void Construct(const FArguments& InArgs, const TSharedRef<STableViewBase>& InOwner, const TSharedRef<UE::SettingsEditor::Private::FTreeItem>& InItem);
+
+private:
+	TWeakPtr<UE::SettingsEditor::Private::FTreeItem> ItemWeak;
+};
\ No newline at end of file
diff --git a/Engine/Source/Developer/SettingsEditor/Private/Widgets/SSettingsEditorCategoryTree.cpp b/Engine/Source/Developer/SettingsEditor/Private/Widgets/SSettingsEditorCategoryTree.cpp
new file mode 100644
--- /dev/null
+++ b/Engine/Source/Developer/SettingsEditor/Private/Widgets/SSettingsEditorCategoryTree.cpp
@@ -0,0 +1,301 @@
+// Copyright Epic Games, Inc. All Rights Reserved.
+
+#include "Widgets/SSettingsEditorCategoryTree.h"
+
+#include "Containers/Ticker.h"
+#include "Internationalization/Internationalization.h"
+#include "ISettingsCategory.h"
+#include "ISettingsContainer.h"
+#include "ISettingsSection.h"
+#include "Styling/AppStyle.h"
+#include "Widgets/Views/STreeView.h"
+
+using namespace UE::SettingsEditor::Private;
+
+void SSettingsEditorCategoryTree::Construct(const FArguments& InArgs, const TSharedPtr<ISettingsContainer>& InSettingsContainer)
+{
+	checkf(InSettingsContainer.IsValid(), TEXT("Settings container must be valid"));
+	SettingsContainer = InSettingsContainer;
+	OnSelectionChanged = InArgs._OnSelectionChanged;
+
+	ChildSlot
+	.VAlign(VAlign_Fill)
+	[
+		SAssignNew(TreeView, STreeView<TSharedPtr<FTreeItem>>)
+		.TreeViewStyle(&FAppStyle::Get().GetWidgetStyle<FTableViewStyle>("SimpleListView"))
+		.SelectionMode(ESelectionMode::Single)
+		.ClearSelectionOnClick(false)
+		.TreeItemsSource(&TreeItemsSource)
+		.OnIsSelectableOrNavigable(this, &SSettingsEditorCategoryTree::OnIsSelectable)
+		.OnGenerateRow(this, &SSettingsEditorCategoryTree::OnTreeGenerateRow) 
+		.OnGetChildren(this, &SSettingsEditorCategoryTree::OnTreeGetChildren)
+		.OnSelectionChanged(this, &SSettingsEditorCategoryTree::OnTreeSelectionChanged)
+	];
+
+	UpdateTreeItemsSource();
+
+	if (InArgs._InitialSelection.IsValid())
+	{
+		SetSelectionInternal(InArgs._InitialSelection);
+	}
+	else
+	{
+		SetSelectionInternal(/** Section */nullptr);
+	}
+
+	SettingsContainer->OnCategoryModified().AddSP(this, &SSettingsEditorCategoryTree::OnCategoryModified);
+	FInternationalization::Get().OnCultureChanged().AddSP(this, &SSettingsEditorCategoryTree::OnCultureChanged);
+}
+
+SSettingsEditorCategoryTree::~SSettingsEditorCategoryTree()
+{
+	if (SettingsContainer)
+	{
+		SettingsContainer->OnCategoryModified().RemoveAll(this);
+	}
+	FInternationalization::Get().OnCultureChanged().RemoveAll(this);
+}
+
+void SSettingsEditorCategoryTree::SetSelection(TSharedPtr<ISettingsSection> InSection)
+{
+	if (TreeView->GetNumItemsSelected() == 1 && GetSelectedSection() == InSection)
+	{
+		return;
+	}
+
+	SetSelectionInternal(InSection);
+}
+
+TSharedPtr<ISettingsSection> SSettingsEditorCategoryTree::GetSelectedSection() const
+{
+	TArray<TSharedPtr<FTreeItem>> SelectedItems = TreeView->GetSelectedItems();
+	return SelectedItems.IsValidIndex(0) && SelectedItems[0].IsValid() ? SelectedItems[0]->Section : nullptr;
+}
+
+void SSettingsEditorCategoryTree::ForEachSection(TFunctionRef<bool(const TSharedPtr<ISettingsSection>&)> InFunctor)
+{
+	for (const TSharedPtr<FTreeItem>& TreeItem : TreeItemsSource)
+	{
+		if (!TreeItem.IsValid() || !TreeItem->Category.IsValid())
+		{
+			continue;
+		}
+
+		// Load children once if needed
+		if (TreeItem->Children.IsEmpty())
+		{
+			OnTreeGetChildren(TreeItem, TreeItem->Children);
+		}
+
+		for (const TSharedPtr<FTreeItem>& ChildItem : TreeItem->Children)
+		{
+			if (!ChildItem.IsValid() || !ChildItem->Section.IsValid())
+			{
+				continue;
+			}
+				
+			if (!InFunctor(ChildItem->Section))
+			{
+				return;
+			}
+		}
+	}
+}
+
+TSharedRef<ITableRow> SSettingsEditorCategoryTree::OnTreeGenerateRow(TSharedPtr<FTreeItem> InItem, const TSharedRef<STableViewBase>& InOwnerTable)
+{
+	return SNew(SSettingsEditorCategoryItem, InOwnerTable, InItem.ToSharedRef());
+}
+
+void SSettingsEditorCategoryTree::OnTreeGetChildren(TSharedPtr<FTreeItem> InItem, TArray<TSharedPtr<FTreeItem>>& OutChildren)
+{
+	// Only process if Category is valid and Section is not yet set
+	if (!InItem.IsValid() || !InItem->Category.IsValid() || InItem->Section.IsValid())
+	{
+		return;
+	}
+	
+	// Children have already been loaded previously
+	if (!InItem->Children.IsEmpty())
+	{
+		OutChildren = InItem->Children;
+		return;
+	}
+
+	TArray<TSharedPtr<ISettingsSection>> CategorySections;
+	InItem->Category->GetSections(CategorySections);
+
+	// sort the sections alphabetically
+	struct FSectionSortPredicate
+	{
+		FORCEINLINE bool operator()(const TSharedPtr<ISettingsSection>& InSectionA, const TSharedPtr<ISettingsSection>& InSectionB) const
+		{
+			if (!InSectionA.IsValid() && !InSectionB.IsValid())
+			{
+				return false;
+			}
+
+			if (InSectionA.IsValid() != InSectionB.IsValid())
+			{
+				return InSectionB.IsValid();
+			}
+
+			return (InSectionA->GetDisplayName().CompareTo(InSectionB->GetDisplayName()) < 0);
+		}
+	};
+
+	CategorySections.Sort(FSectionSortPredicate());
+
+	InItem->Children.Reserve(CategorySections.Num());
+	for (const TSharedPtr<ISettingsSection>& CategorySection : CategorySections)
+	{
+		if (!CategorySection.IsValid())
+		{
+			continue;
+		}
+		
+		TSharedPtr<FTreeItem> NewTreeItem = MakeShared<FTreeItem>();
+		NewTreeItem->Category = InItem->Category;
+		NewTreeItem->Section = CategorySection;
+		InItem->Children.Add(NewTreeItem);
+	}
+
+	OutChildren = InItem->Children;
+}
+
+void SSettingsEditorCategoryTree::OnTreeSelectionChanged(TSharedPtr<FTreeItem> InItem, ESelectInfo::Type InSelectInfo)
+{
+	if (!InItem.IsValid())
+	{
+		return;
+	}
+
+	OnSelectionChanged.ExecuteIfBound(InItem->Section);
+}
+
+bool SSettingsEditorCategoryTree::OnIsSelectable(TSharedPtr<UE::SettingsEditor::Private::FTreeItem> InItem)
+{
+	if (!InItem.IsValid())
+	{
+		return false;
+	}
+
+	// "All settings" option will not have any valid section or category, the rest should have both
+	const bool bValidSection = InItem->Category.IsValid() && InItem->Section.IsValid();
+	const bool bAllSections = !InItem->Category.IsValid() && !InItem->Section.IsValid();
+	return bValidSection || bAllSections;
+}
+
+void SSettingsEditorCategoryTree::UpdateTreeItemsSource()
+{
+	TreeItemsSource.Reset();
+
+	if (!SettingsContainer.IsValid())
+	{
+		return;
+	}
+
+	TArray<TSharedPtr<ISettingsCategory>> Categories;
+	SettingsContainer->GetCategories(Categories);
+
+	TreeItemsSource.Reserve(Categories.Num() + 1);
+
+	// All section entry
+	{
+		TSharedPtr<FTreeItem> AllTreeItem = MakeShared<FTreeItem>();
+		AllTreeItem->Category = nullptr;
+		AllTreeItem->Section = nullptr;
+		TreeItemsSource.Add(AllTreeItem);
+	}
+
+	for (const TSharedPtr<ISettingsCategory>& Category : Categories)
+	{
+		if (!Category.IsValid())
+		{
+			continue;
+		}
+
+		TSharedPtr<FTreeItem> NewTreeItem = MakeShared<FTreeItem>();
+		NewTreeItem->Category = Category;
+		NewTreeItem->Section = nullptr;
+		TreeItemsSource.Add(NewTreeItem);
+	}
+
+	TreeView->RequestTreeRefresh();
+
+	// Expand everything to match legacy behavior
+	for (const TSharedPtr<FTreeItem>& Item : TreeItemsSource)
+	{
+		TreeView->SetItemExpansion(Item.ToSharedRef(), true);
+	}
+}
+
+void SSettingsEditorCategoryTree::OnCultureChanged()
+{
+	RefreshTreeItems();
+}
+
+void SSettingsEditorCategoryTree::OnCategoryModified(const FName& InCategory)
+{
+	if (CategoryUpdateDelegate.IsValid())
+	{
+		return;
+	}
+
+	CategoryUpdateDelegate = FTSTicker::GetCoreTicker().AddTicker(
+		FTickerDelegate::CreateSPLambda(this, 
+			[this](float)->bool
+			{
+				CategoryUpdateDelegate.Reset();
+				RefreshTreeItems();
+				return false; // stop
+			}
+		)
+	);
+}
+
+void SSettingsEditorCategoryTree::RefreshTreeItems()
+{
+	const TSharedPtr<ISettingsSection> SelectedSection = GetSelectedSection();
+	UpdateTreeItemsSource();
+	SetSelectionInternal(SelectedSection);
+}
+
+void SSettingsEditorCategoryTree::SetSelectionInternal(TSharedPtr<ISettingsSection> InSection)
+{
+	// Select "All" item
+	if (!InSection.IsValid())
+	{
+		if (TreeItemsSource.IsValidIndex(0))
+		{
+			TreeView->SetSelection(TreeItemsSource[0], ESelectInfo::Direct);
+		}
+		return;
+	}
+
+	const TSharedPtr<ISettingsCategory> SettingsCategory = InSection->GetCategory().Pin();
+	for (const TSharedPtr<FTreeItem>& TreeItem : TreeItemsSource)
+	{
+		if (!TreeItem.IsValid())
+		{
+			continue;
+		}
+
+		if (TreeItem->Category == SettingsCategory)
+		{
+			for (const TSharedPtr<FTreeItem>& ChildItem : TreeItem->Children)
+			{
+				if (!ChildItem.IsValid())
+				{
+					continue;
+				}
+				
+				if (ChildItem->Section == InSection)
+				{
+					TreeView->SetSelection(ChildItem, ESelectInfo::Direct);
+					return;
+				}
+			}
+		}
+	}
+}
diff --git a/Engine/Source/Developer/SettingsEditor/Private/Widgets/SSettingsEditorCategoryTree.h b/Engine/Source/Developer/SettingsEditor/Private/Widgets/SSettingsEditorCategoryTree.h
new file mode 100644
--- /dev/null
+++ b/Engine/Source/Developer/SettingsEditor/Private/Widgets/SSettingsEditorCategoryTree.h
@@ -0,0 +1,64 @@
+// Copyright Epic Games, Inc. All Rights Reserved.
+
+#pragma once
+
+#include "Containers/Ticker.h"
+#include "SSettingsEditorCategoryItem.h"
+#include "Widgets/SCompoundWidget.h"
+
+class ISettingsContainer;
+class ISettingsCategory;
+class ISettingsSection;
+class ITableRow;
+class STableViewBase;
+namespace ESelectInfo { enum Type : int; }
+template <typename ItemType> class STreeView;
+
+/**
+ * Tree view that displays all the editor settings categories
+ */
+class SSettingsEditorCategoryTree : public SCompoundWidget
+{
+public:
+	DECLARE_DELEGATE_OneParam(FOnSelectionChanged, TSharedPtr<ISettingsSection> /** Section */)
+
+	SLATE_BEGIN_ARGS(SSettingsEditorCategoryTree)
+	{}
+		/** Delegate for when selection has changed */
+		SLATE_EVENT(FOnSelectionChanged, OnSelectionChanged)
+		/** Initial selection */
+		SLATE_ARGUMENT(TSharedPtr<ISettingsSection>, InitialSelection)
+	SLATE_END_ARGS()
+
+	void Construct(const FArguments& InArgs, const TSharedPtr<ISettingsContainer>& InSettingsContainer);
+	virtual ~SSettingsEditorCategoryTree() override;
+	
+	void SetSelection(TSharedPtr<ISettingsSection> InSection);
+	TSharedPtr<ISettingsSection> GetSelectedSection() const;
+
+	void ForEachSection(TFunctionRef<bool(const TSharedPtr<ISettingsSection>&)> InFunctor);
+
+private:
+	TSharedRef<ITableRow> OnTreeGenerateRow(TSharedPtr<UE::SettingsEditor::Private::FTreeItem> InItem, const TSharedRef<STableViewBase>& InOwnerTable);
+	void OnTreeGetChildren(TSharedPtr<UE::SettingsEditor::Private::FTreeItem> InItem, TArray<TSharedPtr<UE::SettingsEditor::Private::FTreeItem>>& OutChildren);
+	void OnTreeSelectionChanged(TSharedPtr<UE::SettingsEditor::Private::FTreeItem> InItem, ESelectInfo::Type InSelectInfo);
+	bool OnIsSelectable(TSharedPtr<UE::SettingsEditor::Private::FTreeItem> InItem);
+	void UpdateTreeItemsSource();
+
+	/** Callback for when the user's culture has changed. */
+	void OnCultureChanged();
+
+	/** Callback for the modification of categories in the settings container. */
+	void OnCategoryModified(const FName& InCategory);
+	
+	void RefreshTreeItems();
+	
+	void SetSelectionInternal(TSharedPtr<ISettingsSection> InSection);
+
+	TSharedPtr<STreeView<TSharedPtr<UE::SettingsEditor::Private::FTreeItem>>> TreeView;
+	TArray<TSharedPtr<UE::SettingsEditor::Private::FTreeItem>> TreeItemsSource;
+	TSharedPtr<ISettingsContainer> SettingsContainer;
+	FOnSelectionChanged OnSelectionChanged;
+	FTSTicker::FDelegateHandle CategoryUpdateDelegate;
+};
+
diff --git a/Engine/Source/Developer/SettingsEditor/Private/Widgets/SSettingsSectionHeader.cpp b/Engine/Source/Developer/SettingsEditor/Private/Widgets/SSettingsSectionHeader.cpp
--- a/Engine/Source/Developer/SettingsEditor/Private/Widgets/SSettingsSectionHeader.cpp
+++ b/Engine/Source/Developer/SettingsEditor/Private/Widgets/SSettingsSectionHeader.cpp
@@ -587,7 +587,18 @@ TSharedPtr<SWidget> FSettingsDetailRootObjectCustomization::CustomizeObjectHeade
 
 bool FSettingsDetailRootObjectCustomization::AreObjectsVisible(const FDetailsObjectSet& InRootObjectSet) const
 {
-	return SelectedSectionObject == nullptr || SelectedSectionObject == InRootObjectSet.RootObjects[0] || DetailsView.Pin()->HasActiveSearch();
+	if (SelectedSectionObject == nullptr)
+	{
+		// When viewing "All settings", do not show keyboard shortcut settings...
+		if (InRootObjectSet.RootObjects[0]->GetClass()->GetName() == TEXT("EditorKeyboardShortcutSettings"))
+		{
+			return false;
+		}
+		
+		return true;
+	}
+
+	return SelectedSectionObject == InRootObjectSet.RootObjects[0];
 }
 
 bool FSettingsDetailRootObjectCustomization::ShouldDisplayHeader(const FDetailsObjectSet& InRootObjects) const
diff --git a/Engine/Source/Developer/SettingsEditor/Public/ISettingsEditorModule.h b/Engine/Source/Developer/SettingsEditor/Public/ISettingsEditorModule.h
--- a/Engine/Source/Developer/SettingsEditor/Public/ISettingsEditorModule.h
+++ b/Engine/Source/Developer/SettingsEditor/Public/ISettingsEditorModule.h
@@ -67,35 +67,83 @@ public:
 
 	struct FSearchTerm
 	{
-		/** In what context is this search term used (i.e. project/editor) */
-		FName ContainerName;
-		
-		/** The key used to show suggestions values (i.e. cvar/keybind) */
-		FString Key;
-		
-		/** The value used for the actual search along the key */
-		FString Value;
+		/** Used to show this search term in the UI */
+		FText Label;
+
+		/** Tooltip text to show on the filter button */
+		FText Tooltip;
+
+		/** Name of the style set to query icon from */
+		FName IconStyleSetName;
+
+		/** Name of the resource to use for the icon */
+		FName IconStyleName;
+
+		/** Results that should be listed and suggested for this search term */
+		TSet<FString> Values;
+
+		/** Only show this search term and filter in those sections */
+		TSet<FName> AllowedSections;
+
+		/** Do not show this search term and filter in those sections */
+		TSet<FName> DisallowedSections;
+
+		/** Show/hide filter toggle button in the UI or only show suggestions */
+		bool bShowFilter = true;
+
+		/** Returns whether this search term is allowed in a section */
+		bool PassesSectionFilter(FName InSection) const
+		{
+			if (DisallowedSections.Contains(InSection))
+			{
+				return false;
+			}
+
+			if (AllowedSections.IsEmpty() || AllowedSections.Contains(InSection))
+			{
+				return true;
+			}
+
+			return false;
+		}
 	};
 
 	/**
-	 * Registers search terms for suggestions
-	 * 
-	 * @param InSearchTerm The context, key and value to register
+	 * Registers a search term for a specific container and key, does nothing if it already exists
+	 * @param InContainerName The context for this search term (i.e. project/editor)
+	 * @param InKey The key used to show filter/suggestions values (i.e. cvar/keybind)
+	 * @return bool whether the search term is registered
 	 */
-	virtual void RegisterSearchTerm(const FSearchTerm& InSearchTerm) = 0;
+	virtual bool RegisterSearchTerm(FName InContainerName, const FString& InKey) = 0;
 
 	/** 
 	 * Delegate called when new search terms have been registered
 	 */
-	DECLARE_MULTICAST_DELEGATE_OneParam(FOnSearchTermsChanged, const FSearchTerm& /** SearchTerm */)
+	DECLARE_MULTICAST_DELEGATE_TwoParams(FOnSearchTermsChanged, FName /** InContainerName */, const FString& /** InKey */)
 	virtual FOnSearchTermsChanged::RegistrationType& OnSearchTermsChanged() = 0;
+
+	/**
+	 * Finds existing search term for a specific container and key
+	 * @param InContainerName The context for this search term (i.e. project/editor)
+	 * @param InKey The key used to show filter/suggestions values (i.e. cvar/keybind)
+	 * @return FSearchTerm* The search term for this container and key if any
+	 */
+	virtual FSearchTerm* FindSearchTerm(FName InContainerName, const FString& InKey) = 0;
 	
 	/**
-	 * Find search terms based on container name
-	 * 
-	 * @param InContainerName The container name for context
+	 * Finds existing search term for a specific container and key (const version)
+	 * @param InContainerName The context for this search term (i.e. project/editor)
+	 * @param InKey The key used to show filter/suggestions values (i.e. cvar/keybind)
+	 * @return const FSearchTerm* The search term for this container and key if any
+	 */
+	virtual const FSearchTerm* FindSearchTerm(FName InContainerName, const FString& InKey) const = 0;
+
+	/**
+	 * Finds existing search terms for a specific container
+	 * @param InContainerName The context for this search term (i.e. project/editor)
+	 * @return TMap<FString, FSearchTerm>* The search terms for this container if any
 	 */
-	virtual const TMap<FString, TSet<FString>>* FindSearchTerms(FName InContainerName) const = 0;
+	virtual const TMap<FString, FSearchTerm>* FindSearchTerms(FName InContainerName) const = 0;
 
 	/** Virtual destructor. */
 	virtual ~ISettingsEditorModule() override = default;
diff --git a/Engine/Source/Editor/EditorStyle/Private/StarshipStyle.cpp b/Engine/Source/Editor/EditorStyle/Private/StarshipStyle.cpp
--- a/Engine/Source/Editor/EditorStyle/Private/StarshipStyle.cpp
+++ b/Engine/Source/Editor/EditorStyle/Private/StarshipStyle.cpp
@@ -1722,6 +1722,35 @@ void FStarshipEditorStyle::FStyle::SetupGeneralStyles()
 		Set( "HardwareTargeting.ScalableQuality", new IMAGE_BRUSH( "/Icons/HardwareTargeting/ScalableQuality", Icon64x64) );
 		Set( "HardwareTargeting.GraphicsUnspecified", new IMAGE_BRUSH( "/Icons/HardwareTargeting/GraphicsUnspecified", Icon64x64) );
 	}
+	
+	// SettingsEditor
+	{
+		Set( "SettingsEditor.ListBorder", new FSlateRoundedBoxBrush(FStyleColors::Recessed, 4.0f) );
+
+		constexpr float PaddingAmount = 2.0f;
+		constexpr float ExtraTopPadding = 12.f;
+		constexpr float ExtraBottomPadding = 8.f;
+		constexpr float AllPluginsExtraTopPadding = 9.f;
+		constexpr float AllPluginsExtraBottomPadding = 7.f;
+
+		Set( "CategoryTreeItem.PaddingAmount", PaddingAmount);
+		Set( "CategoryTreeItem.Root.BackgroundBrush", new FSlateNoResource );
+		Set( "CategoryTreeItem.Root.BackgroundPadding", FMargin(PaddingAmount, PaddingAmount + ExtraTopPadding, PaddingAmount, PaddingAmount + ExtraBottomPadding) );
+		Set( "CategoryTreeItem.Root.AllBackgroundPadding", FMargin(PaddingAmount, PaddingAmount + AllPluginsExtraTopPadding, PaddingAmount, PaddingAmount + AllPluginsExtraBottomPadding) );
+
+		const FTextBlockStyle ButtonStyleText = FAppStyle::Get().GetWidgetStyle<FTextBlockStyle>("PrimaryButtonText");
+		FTextBlockStyle Text = FTextBlockStyle(ButtonStyleText);
+		Text.ColorAndOpacity = FStyleColors::Foreground;
+		Text.TransformPolicy = ETextTransformPolicy::ToUpper;
+		Set( "CategoryTreeItem.Root.Text", Text );
+		Set( "CategoryTreeItem.BackgroundBrush", new FSlateNoResource );
+		Set( "CategoryTreeItem.BackgroundPadding", FMargin(PaddingAmount) );
+
+		const FTextBlockStyle NormalStyleText = FAppStyle::Get().GetWidgetStyle<FTextBlockStyle>("NormalText");
+		FTextBlockStyle CategoryText = FTextBlockStyle(NormalStyleText);
+		CategoryText.ColorAndOpacity = FStyleColors::Foreground;
+		Set( "CategoryTreeItem.Text", CategoryText );
+	}
 
 #endif // WITH_EDITOR
 
diff --git a/Engine/Source/Editor/InputBindingEditor/Private/InputBindingEditorModule.cpp b/Engine/Source/Editor/InputBindingEditor/Private/InputBindingEditorModule.cpp
--- a/Engine/Source/Editor/InputBindingEditor/Private/InputBindingEditorModule.cpp
+++ b/Engine/Source/Editor/InputBindingEditor/Private/InputBindingEditorModule.cpp
@@ -21,6 +21,7 @@
 #include "HAL/PlatformMisc.h"
 #include "IDetailCustomization.h"
 #include "IDetailGroup.h"
+#include "ISettingsContainer.h"
 #include "ISettingsEditorModule.h"
 #include "ISettingsEditorModule.h"
 #include "ISettingsEditorModule.h"
@@ -227,20 +228,47 @@ public:
 	void UpdateUI()
 	{
 		FInputBindingManager& InputBindingManager = FInputBindingManager::Get();
-		
-		ISettingsEditorModule* SettingsEditorModule = FModuleManager::GetModulePtr<ISettingsEditorModule>(TEXT("SettingsEditor"));
-		FName SearchContext = NAME_None;
-		if (TSharedPtr<IDetailsView> DetailsView = DetailBuilder->GetDetailsViewSharedPtr())
-		{
-			SearchContext = DetailsView->GetIdentifier();
-		}
 
 		static const FText SearchKeyKeybindUnbound = LOCTEXT("SearchKeyKeybindUnbound", "Unbound");
-		static const FText SearchKeyKeybind = LOCTEXT("SearchKeyKeybind", "KeyBind");
-		static const FText SearchKeyContext = LOCTEXT("SearchKeyContext", "Context");
+		static const FText SearchKeyKeybind = LOCTEXT("SearchKeyKeybind", "KeyBind:");
+		static const FText SearchKeyContext = LOCTEXT("SearchKeyContext", "Context:");
 		static const FTextFormat TokenCombineFormat = INVTEXT("{0} {1}");
-		static const FTextFormat TokenDelimiter = INVTEXT("\"{0}\"");
-		static const FTextFormat SearchFormat = INVTEXT("{0}={1}");
+		static const FTextFormat TokenDelimiter = INVTEXT("\"{0}\""); // Use for strict search instead of partial...
+
+		ISettingsEditorModule::FSearchTerm* KeybindSearchTerm = nullptr;
+		ISettingsEditorModule::FSearchTerm* ContextSearchTerm = nullptr;
+		if (ISettingsEditorModule* SettingsEditorModule = FModuleManager::GetModulePtr<ISettingsEditorModule>(TEXT("SettingsEditor")))
+		{
+			if (const TSharedPtr<IDetailsView> DetailsView = DetailBuilder->GetDetailsViewSharedPtr())
+			{
+				const FName SearchContext = DetailsView->GetIdentifier();
+				
+				// Register first
+				const bool bKeybindRegistered = SettingsEditorModule->RegisterSearchTerm(SearchContext, SearchKeyKeybind.ToString());
+				const bool bContextRegistered = SettingsEditorModule->RegisterSearchTerm(SearchContext, SearchKeyContext.ToString());
+
+				if (bKeybindRegistered)
+				{
+					KeybindSearchTerm = SettingsEditorModule->FindSearchTerm(SearchContext, SearchKeyKeybind.ToString());
+					KeybindSearchTerm->Label = LOCTEXT("SearchKeyKeybindLabel", "Keybind");
+					KeybindSearchTerm->Tooltip = LOCTEXT("SearchKeyKeybindTooltip", "Filters by the pressed key combination");
+					KeybindSearchTerm->IconStyleSetName = FAppStyle::GetAppStyleSetName();
+					KeybindSearchTerm->IconStyleName = TEXT("GraphEditor.KeyEvent_16x");
+					KeybindSearchTerm->AllowedSections.Add(TEXT("InputBindings")); // Only Keyboard Shortcuts
+				}
+
+				if (bContextRegistered)
+				{
+					ContextSearchTerm = SettingsEditorModule->FindSearchTerm(SearchContext, SearchKeyContext.ToString());
+					ContextSearchTerm->Label = LOCTEXT("SearchKeyContextLabel", "Context");
+					ContextSearchTerm->Tooltip = LOCTEXT("SearchKeyContextTooltip", "Filters by the command context");
+					// Hide this option from being shown in filter buttons but allow it to be suggested
+					ContextSearchTerm->bShowFilter = false;
+					ContextSearchTerm->AllowedSections.Add(TEXT("InputBindings")); // Only Keyboard Shortcuts
+				}
+			}
+		}
+
 		for (TSharedPtr<FChordTreeItem>& TreeItem : ContextList)
 		{
 			check(TreeItem->IsContext());
@@ -261,6 +289,12 @@ public:
 			
 			Commands.Sort(FChordSort(true, false));
 
+			if (KeybindSearchTerm)
+			{
+				// Commands * 2 since there are primary and secondary keybinds
+				KeybindSearchTerm->Values.Reserve(KeybindSearchTerm->Values.Num() + (Commands.Num() * 2));
+			}
+
 			TMap<FName, IDetailGroup*> BundleMap;
 			for(TSharedPtr<FUICommandInfo>& CommandInfo : Commands)
 			{
@@ -295,7 +329,7 @@ public:
 				}
 
 				{
-					// Set up search filter (for i.e. keybind="F")
+					// Set up search filter (for i.e. keybind: "F")
 					const TSharedRef<const FInputChord> FirstInputChord = CommandInfo->GetActiveChord(EMultipleKeyBindingIndex::Primary);
 					const TSharedRef<const FInputChord> SecondInputChord = CommandInfo->GetActiveChord(EMultipleKeyBindingIndex::Secondary);
 
@@ -305,37 +339,23 @@ public:
 					FText FirstKeyFilter, SecondKeyFilter;
 					if (bFirstChordValid)
 					{
-						const FText InputText = FText::FormatOrdered(TokenDelimiter, FirstInputChord->GetInputText());
-						FirstKeyFilter = FText::FormatOrdered(SearchFormat, SearchKeyKeybind.ToLower(), InputText);
+						FirstKeyFilter = FText::FormatOrdered(TokenDelimiter, FirstInputChord->GetInputText());
 
 						// Register search key and value for settings editor suggestions
-						if (SettingsEditorModule)
+						if (KeybindSearchTerm)
 						{
-							ISettingsEditorModule::FSearchTerm SearchTerm
-							{
-								.ContainerName = SearchContext,
-								.Key = SearchKeyKeybind.ToString(),
-								.Value = InputText.ToString()
-							};
-							SettingsEditorModule->RegisterSearchTerm(SearchTerm);
+							KeybindSearchTerm->Values.Add(FirstInputChord->GetInputText().ToString());
 						}
 					}
 
 					if (bSecondChordValid)
 					{
-						const FText InputText = FText::FormatOrdered(TokenDelimiter, SecondInputChord->GetInputText());
-						SecondKeyFilter = FText::FormatOrdered(SearchFormat, SearchKeyKeybind.ToLower(), InputText);
+						SecondKeyFilter = FText::FormatOrdered(TokenDelimiter, SecondInputChord->GetInputText());
 
 						// Register search key and value for settings editor suggestions
-						if (SettingsEditorModule)
+						if (KeybindSearchTerm)
 						{
-							ISettingsEditorModule::FSearchTerm SearchTerm
-							{
-								.ContainerName = SearchContext,
-								.Key = SearchKeyKeybind.ToString(),
-								.Value = InputText.ToString()
-							};
-							SettingsEditorModule->RegisterSearchTerm(SearchTerm);
+							KeybindSearchTerm->Values.Add(SecondInputChord->GetInputText().ToString());
 						}
 					}
 
@@ -343,21 +363,15 @@ public:
 						? FText::FormatOrdered(TokenCombineFormat, FirstKeyFilter, SecondKeyFilter)
 						: (bFirstChordValid ? FirstKeyFilter : SecondKeyFilter);
 
-					// Allow search for unbound commands like this : keybind=unbound
+					// Allow search for unbound commands like this : keybind: "unbound"
 					if (FilterText.IsEmpty())
 					{
-						FilterText = FText::FormatOrdered(SearchFormat, SearchKeyKeybind.ToLower(), SearchKeyKeybindUnbound.ToLower());
+						FilterText = FText::FormatOrdered(TokenDelimiter, SearchKeyKeybindUnbound.ToLower());
 
 						// Register search key and value for settings editor suggestions
-						if (SettingsEditorModule)
+						if (KeybindSearchTerm)
 						{
-							ISettingsEditorModule::FSearchTerm SearchTerm
-							{
-								.ContainerName = SearchContext,
-								.Key = SearchKeyKeybind.ToString(),
-								.Value = SearchKeyKeybindUnbound.ToString()
-							};
-							SettingsEditorModule->RegisterSearchTerm(SearchTerm);
+							KeybindSearchTerm->Values.Add(SearchKeyKeybindUnbound.ToString());
 						}
 					}
 
@@ -370,23 +384,20 @@ public:
 					// Append command binding context
 					{
 						FText BindingContext = FText::FormatOrdered(TokenDelimiter, FText::FromName(BindingContextName));
-						
+
 						// Register search key and value for settings editor suggestions
-						if (SettingsEditorModule)
+						if (ContextSearchTerm)
 						{
-							ISettingsEditorModule::FSearchTerm SearchTerm
-							{
-								.ContainerName = SearchContext,
-								.Key = SearchKeyContext.ToString(),
-								.Value = BindingContext.ToString()
-							};
-							SettingsEditorModule->RegisterSearchTerm(SearchTerm);
+							ContextSearchTerm->Values.Add(BindingContextName.ToString());
 						}
 
-						BindingContext = FText::FormatOrdered(SearchFormat, SearchKeyContext.ToLower(), BindingContext);
 						FilterText = FText::FormatOrdered(TokenCombineFormat, FilterText, BindingContext);
 					}
 
+					// Append keywords
+					FilterText = FText::FormatOrdered(TokenCombineFormat, FilterText, SearchKeyKeybind.ToLower());
+					FilterText = FText::FormatOrdered(TokenCombineFormat, FilterText, SearchKeyContext.ToLower());
+
 					Row->FilterString(FilterText);
 				}
 
@@ -455,12 +466,24 @@ public:
 		PropertyEditor.RegisterCustomClassLayout(EditorKeyboardShortcutSettingsName, FOnGetDetailCustomizationInstance::CreateStatic(&FEditorKeyboardShortcutSettings::MakeInstance));
 
 		// input bindings
-		ISettingsSectionPtr InputBindingSettingsSection = SettingsModule.RegisterSettings("Editor", "General", "InputBindings",
+		const FName ContainerName = "Editor";
+		const FName CategoryName = "Shortcuts";
+		const ISettingsSectionPtr InputBindingSettingsSection = SettingsModule.RegisterSettings(
+			ContainerName, 
+			CategoryName, 
+			"InputBindings",
 			LOCTEXT("InputBindingsSettingsName", "Keyboard Shortcuts"),
 			LOCTEXT("InputBindingsSettingsDescription", "Configure keyboard shortcuts to quickly invoke operations."),
 			GetMutableDefault<UEditorKeyboardShortcutSettings>()
 		);
 
+		if (const TSharedPtr<ISettingsContainer> EditorContainer = SettingsModule.GetContainer(ContainerName))
+		{
+			// Place it below "All Settings"
+			EditorContainer->SetCategorySortPriority(CategoryName, -1);
+			EditorContainer->DescribeCategory(CategoryName, LOCTEXT("EditorShortcutsCategoryName", "Shortcuts"), LOCTEXT("EditorShortcutsCategoryDescription", "Shortcuts settings for the editor"));
+		}
+
 		if(InputBindingSettingsSection.IsValid())
 		{
 			InputBindingSettingsSection->OnExport().BindRaw(this, &FInputBindingEditorModule::HandleInputBindingsExport);
diff --git a/Engine/Source/Editor/UnrealEd/Private/Toolkits/GlobalEditorCommonCommands.cpp b/Engine/Source/Editor/UnrealEd/Private/Toolkits/GlobalEditorCommonCommands.cpp
--- a/Engine/Source/Editor/UnrealEd/Private/Toolkits/GlobalEditorCommonCommands.cpp
+++ b/Engine/Source/Editor/UnrealEd/Private/Toolkits/GlobalEditorCommonCommands.cpp
@@ -159,7 +159,7 @@ void FGlobalEditorCommonCommands::MapActions(TSharedRef<FUICommandList>& Toolkit
 	
 	ToolkitCommands->MapAction(
 		GlobalEditorCommands.OpenKeyboardShortcuts,
-		FExecuteAction::CreateStatic(&FGlobalEditorCommonCommands::OpenSettings, FName("Editor"), FName("General"), FName("InputBindings")));
+		FExecuteAction::CreateStatic(&FGlobalEditorCommonCommands::OpenSettings, FName("Editor"), FName("Shortcuts"), FName("InputBindings")));
 	
 	ToolkitCommands->MapAction(
 		GlobalEditorCommands.OpenProjectSettings,
