EditorUsability : SettingsEditor
- Fixed configuration import by removing irrelevant sections for keybinds to avoid same section+key clash on merge
- Fixed configuration export by allowing duplicated keys in a section when merging files, this was causing keybinds to be lost
- Added notification hyperlink to navigate to exported file

#Synced-CL 51025551

diff --git a/Engine/Source/Developer/SettingsEditor/Private/Widgets/SSettingsEditor.cpp b/Engine/Source/Developer/SettingsEditor/Private/Widgets/SSettingsEditor.cpp
--- a/Engine/Source/Developer/SettingsEditor/Private/Widgets/SSettingsEditor.cpp
+++ b/Engine/Source/Developer/SettingsEditor/Private/Widgets/SSettingsEditor.cpp
@@ -437,15 +437,6 @@ EVisibility SSettingsEditor::HandleSettingsBoxVisibility() const
 	return EVisibility::Hidden;
 }
 
-void ShowNotification(const FText& Text, SNotificationItem::ECompletionState CompletionState)
-{
-	FNotificationInfo Notification(Text);
-	Notification.ExpireDuration = 3.f;
-	Notification.bUseSuccessFailIcons = false;
-
-	FSlateNotificationManager::Get().AddNotification(Notification)->SetCompletionState(CompletionState);
-}
-
 FReply SSettingsEditor::HandleExportButtonClicked()
 {
 	const FString DefaultFileName = FString::Printf(TEXT("%s Settings Backup %s.ini"), *Model->GetSettingsContainer()->GetName().ToString(), *FDateTime::Now().ToString(TEXT("%Y-%m-%d %H%M%S")));
@@ -465,7 +456,7 @@ FReply SSettingsEditor::HandleExportButtonClicked()
 		const FString DestinationFilePath = OutSaveFiles[0];
 		if (GConfig->FindBranch(*DestinationFilePath, DestinationFilePath))
 		{
-			ShowNotification(LOCTEXT("ExportSettingsExistingFileFailure", "Export settings failed. Cannot use a file that is already part of the config system."), SNotificationItem::CS_Fail);
+			SSettingsSectionHeader::ShowNotification(LOCTEXT("ExportSettingsExistingFileFailure", "Export settings failed. Cannot use a file that is already part of the config system."), SNotificationItem::CS_Fail);
 			return FReply::Handled();
 		}
 
@@ -492,7 +483,7 @@ FReply SSettingsEditor::HandleExportButtonClicked()
 			}
 			else
 			{
-				ShowNotification(LOCTEXT("ExportSettingsExistingTempFolderFailure", "Export settings failed. TempExport folder already exists under project Saved folder."), SNotificationItem::CS_Fail);
+				SSettingsSectionHeader::ShowNotification(LOCTEXT("ExportSettingsExistingTempFolderFailure", "Export settings failed. TempExport folder already exists under project Saved folder."), SNotificationItem::CS_Fail);
 				return FReply::Handled();
 			}
 		}
@@ -521,7 +512,7 @@ FReply SSettingsEditor::HandleExportButtonClicked()
 					continue;
 				}
 
-				TempBranch.InMemoryFile.Combine(TempDestinationFilePath);
+				TempBranch.InMemoryFile.Combine(TempDestinationFilePath, /** HandleSymbolCommands */false);
 				IFileManager::Get().Delete(*TempDestinationFilePath);
 
 				TArray<FString> AdditionalFiles;
@@ -542,11 +533,11 @@ FReply SSettingsEditor::HandleExportButtonClicked()
 
 		IFileManager::Get().DeleteDirectory(*TempDestinationFolder, bDeleteDirectoryRequireExists, bDeleteDirectoryRecursively);
 
-		ShowNotification(LOCTEXT("ExportSettingsSuccess", "Export settings succeeded"), SNotificationItem::CS_Success);
+		SSettingsSectionHeader::ShowNotification(LOCTEXT("ExportSettingsSuccess", "Export settings succeeded"), SNotificationItem::CS_Success, DestinationFilePath);
 		return FReply::Handled();
 	}
 
-	ShowNotification(LOCTEXT("ExportSettingsFailure", "Export settings failed"), SNotificationItem::CS_Fail);
+	SSettingsSectionHeader::ShowNotification(LOCTEXT("ExportSettingsFailure", "Export settings failed"), SNotificationItem::CS_Fail);
 	return FReply::Handled();
 }
 
@@ -579,11 +570,11 @@ FReply SSettingsEditor::HandleImportButtonClicked()
 			}
 		}
 
-		ShowNotification(LOCTEXT("ImportSettingsSuccess", "Import settings succeeded"), SNotificationItem::CS_Success);
+		SSettingsSectionHeader::ShowNotification(LOCTEXT("ImportSettingsSuccess", "Import settings succeeded"), SNotificationItem::CS_Success);
 	}
 	else
 	{
-		ShowNotification(LOCTEXT("ImportSettingsFailure", "Import settings failed"), SNotificationItem::CS_Fail);
+		SSettingsSectionHeader::ShowNotification(LOCTEXT("ImportSettingsFailure", "Import settings failed"), SNotificationItem::CS_Fail);
 	}
 
 	return FReply::Handled();
diff --git a/Engine/Source/Developer/SettingsEditor/Private/Widgets/SSettingsSectionHeader.cpp b/Engine/Source/Developer/SettingsEditor/Private/Widgets/SSettingsSectionHeader.cpp
--- a/Engine/Source/Developer/SettingsEditor/Private/Widgets/SSettingsSectionHeader.cpp
+++ b/Engine/Source/Developer/SettingsEditor/Private/Widgets/SSettingsSectionHeader.cpp
@@ -26,6 +26,31 @@
 bool bHideSetAsDefaultButton = false;
 FAutoConsoleVariableRef CVarHideSetAsDefaultButton(TEXT("SettingsEditor.HideSetAsDefaultButton"), bHideSetAsDefaultButton, TEXT("Hide the Settings Editor button to save to default config."));
 
+void SSettingsSectionHeader::ShowNotification(const FText& InText, SNotificationItem::ECompletionState InCompletionState, FString InNavigateTo)
+{
+	FNotificationInfo Notification(InText);
+	Notification.ExpireDuration = 10.f;
+	Notification.bFireAndForget = true;
+	Notification.bUseSuccessFailIcons = false;
+
+	if (!InNavigateTo.IsEmpty())
+	{
+		Notification.Hyperlink = FSimpleDelegate::CreateLambda([InNavigateTo]()
+			{
+				IPlatformFile& PlatformFile = FPlatformFileManager::Get().GetPlatformFile();
+				if (PlatformFile.FileExists(*InNavigateTo) || PlatformFile.DirectoryExists(*InNavigateTo))
+				{
+					// Navigate in explorer
+					FPlatformProcess::ExploreFolder(*InNavigateTo);
+				}
+			}
+		);
+		Notification.HyperlinkText = LOCTEXT("NavigateToFile", "Navigate to file");
+	}
+
+	FSlateNotificationManager::Get().AddNotification(Notification)->SetCompletionState(InCompletionState);
+}
+
 void SSettingsSectionHeader::Construct(const FArguments& InArgs, const UObject* InSettingsObject, ISettingsEditorModelPtr InModel, TSharedPtr<IDetailsView> InDetailsView, const TSharedPtr<ITableRow>& InTableRow)
 {
 	Model = InModel;
@@ -219,7 +244,7 @@ FReply SSettingsSectionHeader::HandleExportButtonClicked()
 		{
 			if(SettingsSection->Export(OutFiles[0]))
 			{
-				ShowNotification(LOCTEXT("ExportSettingsSuccess", "Export settings succeeded"), SNotificationItem::CS_Success);
+				ShowNotification(LOCTEXT("ExportSettingsSuccess", "Export settings succeeded"), SNotificationItem::CS_Success, OutFiles[0]);
 			}
 			else
 			{
@@ -506,15 +531,6 @@ bool SSettingsSectionHeader::MakeDefaultConfigFileWritable()
 	return SettingsHelpers::MakeWritable(AbsoluteConfigFilePath, true);
 }
 
-void SSettingsSectionHeader::ShowNotification(const FText& Text, SNotificationItem::ECompletionState CompletionState) const
-{
-	FNotificationInfo Notification(Text);
-	Notification.ExpireDuration = 3.f;
-	Notification.bUseSuccessFailIcons = false;
-
-	FSlateNotificationManager::Get().AddNotification(Notification)->SetCompletionState(CompletionState);
-}
-
 /** Returns the config file name currently being edited. */
 FString SSettingsSectionHeader::HandleCheckoutNoticeConfigFilePath() const
 {
diff --git a/Engine/Source/Developer/SettingsEditor/Private/Widgets/SSettingsSectionHeader.h b/Engine/Source/Developer/SettingsEditor/Private/Widgets/SSettingsSectionHeader.h
--- a/Engine/Source/Developer/SettingsEditor/Private/Widgets/SSettingsSectionHeader.h
+++ b/Engine/Source/Developer/SettingsEditor/Private/Widgets/SSettingsSectionHeader.h
@@ -18,6 +18,8 @@ class SSettingsEditorCheckoutNotice;
 class SSettingsSectionHeader : public SCompoundWidget
 {
 public:
+	static void ShowNotification(const FText& InText, SNotificationItem::ECompletionState InCompletionState, FString InNavigateTo = FString());
+
 	SLATE_BEGIN_ARGS(SSettingsSectionHeader)
 	{}
 	SLATE_END_ARGS()
@@ -80,8 +82,6 @@ private:
 	*/
 	bool MakeDefaultConfigFileWritable();
 
-	void ShowNotification(const FText& Text, SNotificationItem::ECompletionState CompletionState) const;
-
 	/** Returns the config file name currently being edited. */
 	FString HandleCheckoutNoticeConfigFilePath() const;
 
diff --git a/Engine/Source/Runtime/Core/Private/Misc/ConfigCacheIni.cpp b/Engine/Source/Runtime/Core/Private/Misc/ConfigCacheIni.cpp
--- a/Engine/Source/Runtime/Core/Private/Misc/ConfigCacheIni.cpp
+++ b/Engine/Source/Runtime/Core/Private/Misc/ConfigCacheIni.cpp
@@ -1513,9 +1513,9 @@ const FConfigSection* FConfigFile::FindOrAddConfigSection(const FString& Section
 	return FindOrAddSectionInternal(SectionName);
 }
 
-bool FConfigFile::Combine(const FString& Filename)
+bool FConfigFile::Combine(const FString& Filename, bool bHandleSymbolCommands)
 {
-	return FillFileFromDisk(Filename, true);
+	return FillFileFromDisk(Filename, bHandleSymbolCommands);
 }
 
 void FConfigFile::Shrink()
@@ -2134,9 +2134,9 @@ bool FConfigFile::FillFileFromDisk(const FString& Filename, bool bHandleSymbolCo
 	return ::FillFileFromDisk(this, Filename, bHandleSymbolCommands);
 }
 
-void FConfigFile::CombineFromBuffer(const FString& Buffer, const FString& FileHint)
+void FConfigFile::CombineFromBuffer(const FString& Buffer, const FString& FileHint, bool bHandleSymbolCommands)
 {
-	::FillFileFromBuffer(this, Buffer, true, FileHint);
+	::FillFileFromBuffer(this, Buffer, bHandleSymbolCommands, FileHint);
 }
 
 /**
diff --git a/Engine/Source/Runtime/Core/Public/Misc/ConfigCacheIni.h b/Engine/Source/Runtime/Core/Public/Misc/ConfigCacheIni.h
--- a/Engine/Source/Runtime/Core/Public/Misc/ConfigCacheIni.h
+++ b/Engine/Source/Runtime/Core/Public/Misc/ConfigCacheIni.h
@@ -732,8 +732,8 @@ public:
 	bool operator==( const FConfigFile& Other ) const;
 	bool operator!=( const FConfigFile& Other ) const;
 
-	CORE_API bool Combine(const FString& Filename);
-	CORE_API void CombineFromBuffer(const FString& Buffer, const FString& FileHint);
+	CORE_API bool Combine(const FString& Filename, bool bHandleSymbolCommands = true);
+	CORE_API void CombineFromBuffer(const FString& Buffer, const FString& FileHint, bool bHandleSymbolCommands = true);
 	CORE_API void Read( const FString& Filename );
 
 	/**
diff --git a/Engine/Source/Runtime/Slate/Private/Framework/Commands/InputBindingManager.cpp b/Engine/Source/Runtime/Slate/Private/Framework/Commands/InputBindingManager.cpp
--- a/Engine/Source/Runtime/Slate/Private/Framework/Commands/InputBindingManager.cpp
+++ b/Engine/Source/Runtime/Slate/Private/Framework/Commands/InputBindingManager.cpp
@@ -245,6 +245,16 @@ void FUserDefinedChords::SaveChords() const
 			ChordJsonArray.Add(EscapedContent);
 		}
 
+		// Remove all sections not related to keybinds...
+		TArray<FString> SectionNames;
+		if (GConfig->GetSectionNames(GEditorKeyBindingsIni, SectionNames))
+		{
+			for (const FString& SectionKey : SectionNames)
+			{
+				GConfig->EmptySection(*SectionKey, GEditorKeyBindingsIni);
+			}
+		}
+
 		GConfig->SetArray(TEXT("UserDefinedChords"), TEXT("UserDefinedChords"), ChordJsonArray, GEditorKeyBindingsIni);
 
 		// Clean up old keys (if they still exist)
