SSettingsEditor : observe UndoRedo transaction events on objects that we've observed a property change for and share the logic for saving the property deltas in both cases.

This fixes saving settings to their config files on undo/redo for the general case.

#Synced-CL 51609228

diff --git a/Engine/Source/Developer/SettingsEditor/Private/Widgets/SSettingsEditor.cpp b/Engine/Source/Developer/SettingsEditor/Private/Widgets/SSettingsEditor.cpp
--- a/Engine/Source/Developer/SettingsEditor/Private/Widgets/SSettingsEditor.cpp
+++ b/Engine/Source/Developer/SettingsEditor/Private/Widgets/SSettingsEditor.cpp
@@ -32,6 +32,7 @@
 #include "HAL/FileManager.h"
 #include "Misc/ConfigCacheIni.h"
 #include "Misc/MessageDialog.h"
+#include "Misc/TransactionObjectEvent.h"
 #include "Models/SettingsEditorMenuContext.h"
 #include "Widgets/Layout/SBox.h"
 
@@ -46,6 +47,7 @@ const FTextFormat SSettingsEditor::TokenDelimiter = INVTEXT("\"{0}\"");
 
 SSettingsEditor::~SSettingsEditor()
 {
+	FCoreUObjectDelegates::OnObjectTransacted.RemoveAll(this);
 	Model->OnSelectionChanged().RemoveAll(this);
 	if (ISettingsEditorModule* SettingsEditorModule = FModuleManager::GetModulePtr<ISettingsEditorModule>(TEXT("SettingsEditor")))
 	{
@@ -103,6 +105,8 @@ void SSettingsEditor::Construct( const FArguments& InArgs, const ISettingsEditor
 	RootObjectCustomization->Initialize();
 	SettingsView->SetRootObjectCustomizationInstance(RootObjectCustomization);
 
+	FCoreUObjectDelegates::OnObjectTransacted.AddSP(this, &SSettingsEditor::HandleObjectTransacted);
+
 	// Register menu toolbar
 	UToolMenus* ToolMenus = UToolMenus::Get();
 	if (IsValid(ToolMenus) && !ToolMenus->IsMenuRegistered(SettingsEditorToolbarName))
@@ -261,93 +265,125 @@ void SSettingsEditor::NotifyPostChange( const FPropertyChangedEvent& PropertyCha
 		if (PropertyChangedEvent.GetNumObjectsBeingEdited() > 0)
 		{
 			UObject* ObjectBeingEdited = (UObject*)PropertyChangedEvent.GetObjectBeingEdited(0);
+			const FProperty* MemberProperty = PropertyThatChanged->GetActiveMemberNode()->GetValue();
+			const FProperty* ChangedProperty = PropertyChangedEvent.Property;
 
-			// Get the section from the edited object.  We cannot use the selected section as multiple sections can be shown at once in the settings details panel.
-			ISettingsSectionPtr Section = Model->GetSectionFromSectionObject(ObjectBeingEdited);
-			{
-				FString RelativePath;
-				bool bIsSourceControlled = false;
-				bool bIsNewFile = false;
+			HandleSettingsPropertyChanged(ObjectBeingEdited, MemberProperty, ChangedProperty);
 
-				// Attempt to checkout the file automatically
-				if (ObjectBeingEdited->GetClass()->HasAnyClassFlags(CLASS_DefaultConfig))
-				{
-					RelativePath = ObjectBeingEdited->GetDefaultConfigFilename();
-					bIsSourceControlled = true;
-				}
-				else if (ObjectBeingEdited->GetClass()->HasAnyClassFlags(CLASS_Config))
-				{
-					RelativePath = ObjectBeingEdited->GetClass()->GetConfigName();
-				}
+			ModifiedObjectCache.Add(ObjectBeingEdited);
+		}
+	}
+}
 
-				FString FullPath = FPaths::ConvertRelativePathToFull(RelativePath);
+void SSettingsEditor::HandleSettingsPropertyChanged(UObject* ObjectBeingEdited, const FProperty* MemberProperty, const FProperty* ChangedProperty)
+{
+	// Get the section from the edited object.  We cannot use the selected section as multiple sections can be shown at once in the settings details panel.
+	ISettingsSectionPtr Section = Model->GetSectionFromSectionObject(ObjectBeingEdited);
+	{
+		FString RelativePath;
+		bool bIsSourceControlled = false;
+		bool bIsNewFile = false;
 
-				if (!FPlatformFileManager::Get().GetPlatformFile().FileExists(*FullPath))
-				{
-					bIsNewFile = true;
-				}
+		// Attempt to checkout the file automatically
+		if (ObjectBeingEdited->GetClass()->HasAnyClassFlags(CLASS_DefaultConfig))
+		{
+			RelativePath = ObjectBeingEdited->GetDefaultConfigFilename();
+			bIsSourceControlled = true;
+		}
+		else if (ObjectBeingEdited->GetClass()->HasAnyClassFlags(CLASS_Config))
+		{
+			RelativePath = ObjectBeingEdited->GetClass()->GetConfigName();
+		}
 
-				if (!bIsSourceControlled || !SettingsHelpers::CheckOutOrAddFile(FullPath))
-				{
-					SettingsHelpers::MakeWritable(FullPath);
-				}
-			
-				if (Section.IsValid())
-				{
-					RecordPreferenceChangedAnalytics(Section, PropertyChangedEvent);
-				}
+		FString FullPath = FPaths::ConvertRelativePathToFull(RelativePath);
 
-				static IConsoleVariable* NewPropertySaving = IConsoleManager::Get().FindConsoleVariable(TEXT("ini.UseNewPropertySaving"));
-				bool bUseFullSectionSave = false;
-				if (NewPropertySaving->GetInt() == 0)
-				{
-					// Determine if the Property is an Array or Array Element
-					bool bIsArrayOrArrayElement = PropertyThatChanged->GetActiveMemberNode()->GetValue()->IsA(FArrayProperty::StaticClass()) 
-						|| PropertyThatChanged->GetActiveMemberNode()->GetValue()->ArrayDim > 1
-						|| PropertyChangedEvent.Property->GetOwner<FArrayProperty>();
-	
-					bool bIsSetOrSetElement = PropertyThatChanged->GetActiveMemberNode()->GetValue()->IsA(FSetProperty::StaticClass())
-						|| PropertyChangedEvent.Property->GetOwner<FSetProperty>();
-	
-					bool bIsMapOrMapElement = PropertyThatChanged->GetActiveMemberNode()->GetValue()->IsA(FMapProperty::StaticClass())
-						|| PropertyChangedEvent.Property->GetOwner<FMapProperty>();
+		if (!FPlatformFileManager::Get().GetPlatformFile().FileExists(*FullPath))
+		{
+			bIsNewFile = true;
+		}
 
-					bUseFullSectionSave = bIsArrayOrArrayElement || bIsSetOrSetElement || bIsMapOrMapElement;
-				}
+		if (!bIsSourceControlled || !SettingsHelpers::CheckOutOrAddFile(FullPath))
+		{
+			SettingsHelpers::MakeWritable(FullPath);
+		}
 
-				if (ObjectBeingEdited->GetClass()->HasAnyClassFlags(CLASS_DefaultConfig) && !bUseFullSectionSave)
-				{
-					if (!Section.IsValid() || Section->NotifySectionOnPropertyModified())
-					{
-						ObjectBeingEdited->UpdateSinglePropertyInConfigFile(PropertyThatChanged->GetActiveMemberNode()->GetValue(), ObjectBeingEdited->GetDefaultConfigFilename());
-					}
-				}
-				else if (Section.IsValid())
-				{
-					Section->Save();
-				}
-				// Some files being edited might have an array element, but they may also not have a corresponding section for inlined
-				// external objects, for them if they're DefaultConfig, we update them here.
-				else if (ObjectBeingEdited->GetClass()->HasAnyClassFlags(CLASS_DefaultConfig))
-				{
-					ObjectBeingEdited->TryUpdateDefaultConfigFile();
-				}
+		if (Section.IsValid())
+		{
+			RecordPreferenceChangedAnalytics(Section, MemberProperty);
+		}
 
-				if (bIsNewFile && bIsSourceControlled)
-				{
-					SettingsHelpers::CheckOutOrAddFile(FullPath);
-				}
+		static IConsoleVariable* NewPropertySaving = IConsoleManager::Get().FindConsoleVariable(TEXT("ini.UseNewPropertySaving"));
+		bool bUseFullSectionSave = false;
+		if (NewPropertySaving->GetInt() == 0)
+		{
+			// Determine if the Property is an Array or Array Element
+			bool bIsArrayOrArrayElement = MemberProperty->IsA(FArrayProperty::StaticClass())
+				|| MemberProperty->ArrayDim > 1
+				|| ChangedProperty->GetOwner<FArrayProperty>();
 
-				static const FName ConfigRestartRequiredKey = "ConfigRestartRequired";
-				if (PropertyChangedEvent.Property->GetBoolMetaData(ConfigRestartRequiredKey) || PropertyChangedEvent.MemberProperty->GetBoolMetaData(ConfigRestartRequiredKey))
-				{
-					OnApplicationRestartRequiredDelegate.ExecuteIfBound();
-				}
+			bool bIsSetOrSetElement = MemberProperty->IsA(FSetProperty::StaticClass())
+				|| ChangedProperty->GetOwner<FSetProperty>();
+
+			bool bIsMapOrMapElement = MemberProperty->IsA(FMapProperty::StaticClass())
+				|| ChangedProperty->GetOwner<FMapProperty>();
+
+			bUseFullSectionSave = bIsArrayOrArrayElement || bIsSetOrSetElement || bIsMapOrMapElement;
+		}
+
+		if (ObjectBeingEdited->GetClass()->HasAnyClassFlags(CLASS_DefaultConfig) && !bUseFullSectionSave)
+		{
+			if (!Section.IsValid() || Section->NotifySectionOnPropertyModified())
+			{
+				ObjectBeingEdited->UpdateSinglePropertyInConfigFile(MemberProperty, ObjectBeingEdited->GetDefaultConfigFilename());
 			}
 		}
+		else if (Section.IsValid())
+		{
+			Section->Save();
+		}
+		// Some files being edited might have an array element, but they may also not have a corresponding section for inlined
+		// external objects, for them if they're DefaultConfig, we update them here.
+		else if (ObjectBeingEdited->GetClass()->HasAnyClassFlags(CLASS_DefaultConfig))
+		{
+			ObjectBeingEdited->TryUpdateDefaultConfigFile();
+		}
+
+		if (bIsNewFile && bIsSourceControlled)
+		{
+			SettingsHelpers::CheckOutOrAddFile(FullPath);
+		}
+
+		static const FName ConfigRestartRequiredKey = "ConfigRestartRequired";
+		if (ChangedProperty->GetBoolMetaData(ConfigRestartRequiredKey) || MemberProperty->GetBoolMetaData(ConfigRestartRequiredKey))
+		{
+			OnApplicationRestartRequiredDelegate.ExecuteIfBound();
+		}
 	}
 }
 
+void SSettingsEditor::HandleObjectTransacted(UObject* InObject, const FTransactionObjectEvent& InTransactionEvent)
+{
+	if (InTransactionEvent.GetEventType() != ETransactionObjectEventType::UndoRedo)
+	{
+		return;
+	}
+
+	if (!InTransactionEvent.HasPropertyChanges())
+	{
+		return;
+	}
+
+	if (ModifiedObjectCache.Contains(InObject))
+	{
+		for (const FName& PropertyName : InTransactionEvent.GetChangedProperties())
+		{
+			if (const FProperty* MemberProperty = InObject->GetClass()->FindPropertyByName(PropertyName))
+			{
+				HandleSettingsPropertyChanged(InObject, MemberProperty, MemberProperty);
+			}
+		}
+	}
+}
 
 /* SSettingsEditor implementation
  *****************************************************************************/
@@ -365,9 +401,8 @@ TWeakObjectPtr<UObject> SSettingsEditor::GetSelectedSettingsObject() const
 }
 
 
-void SSettingsEditor::RecordPreferenceChangedAnalytics( ISettingsSectionPtr SelectedSection, const FPropertyChangedEvent& PropertyChangedEvent ) const
+void SSettingsEditor::RecordPreferenceChangedAnalytics( ISettingsSectionPtr SelectedSection, const FProperty* ChangedProperty ) const
 {
-	FProperty* ChangedProperty = PropertyChangedEvent.MemberProperty;
 	// submit analytics data
 	if(FEngineAnalytics::IsAvailable() && ChangedProperty != nullptr && ChangedProperty->GetOwnerClass() != nullptr)
 	{
diff --git a/Engine/Source/Developer/SettingsEditor/Private/Widgets/SSettingsEditor.h b/Engine/Source/Developer/SettingsEditor/Private/Widgets/SSettingsEditor.h
--- a/Engine/Source/Developer/SettingsEditor/Private/Widgets/SSettingsEditor.h
+++ b/Engine/Source/Developer/SettingsEditor/Private/Widgets/SSettingsEditor.h
@@ -15,6 +15,7 @@
 #include "Widgets/SCompoundWidget.h"
 #include "Widgets/SOverlay.h"
 
+class FTransactionObjectEvent;
 class IDetailsView;
 class SAssetSearchBox;
 class SBox;
@@ -78,9 +79,9 @@ protected:
 	 * Reports a preference changed event to the analytics system.
 	 *
 	 * @param SelectedSection The currently selected settings section.
-	 * @param PropertyChangedEvent The event for the changed property.
+	 * @param ChangedProperty The member property that changed.
 	 */
-	void RecordPreferenceChangedAnalytics( ISettingsSectionPtr SelectedSection, const FPropertyChangedEvent& PropertyChangedEvent ) const;
+	void RecordPreferenceChangedAnalytics( ISettingsSectionPtr SelectedSection, const FProperty* ChangedProperty) const;
 
 private:
 	static const FTextFormat TokenCombineFormat;
@@ -89,6 +90,9 @@ private:
 	static void FillSettingsFiltersSection(UToolMenu* InMenu);
 	static FToolMenuEntry CreateFiltersCheckboxEntry(const ISettingsEditorModule* InSettingsEditorModule, const USettingsEditorMenuContext* InContext, const FString& InKey, const TSet<FString>& InValues);
 
+	void HandleSettingsPropertyChanged(UObject* InObjectBeingEdited, const FProperty* InMemberProperty, const FProperty* InChangedProperty);
+	void HandleObjectTransacted(UObject* InObject, const FTransactionObjectEvent& InTransactionEvent);
+
 	/** Callback for changing the selected settings section. */
 	void HandleModelSelectionChanged();
 
@@ -161,4 +165,12 @@ private:
 
 	/** Used to store active search key */
 	FString SearchKey;
+
+	/** 
+	 * Cache of modified objects added to when the user changes properties in the details view.
+	 * Used to filter objects in a FCoreUObjectDelegates::OnObjectTransacted handler and only handle 
+	 * undo/redo events on objects that we've observed the user change a property of in the details view.
+	 * Note that these objects can be settings objects as well as sub-objects of a settings object.
+	 */
+	TSet<TWeakObjectPtr<UObject>> ModifiedObjectCache;
 };
