EditorUsability : SettingsEditor
- Added reset to default support for root properties of Editor Preferences, Project Settings, Keyboard shortcuts
- Added default value in the tooltip of the property/keybind name
- Added CVAR to refresh cached default values : SettingsEditor.ClearObjectDefaultPropertyValuesCache
- Limitations are the following:
* it will not work for properties that are missing the Config specifier
* it will not work for custom added rows obviously
* it will not work for settings object with a customization since it overrides the default one
* it will not work on nested properties, only the root properties of a settings object are supported for performance to avoid having to parse the default stringify value recursively

#Synced-CL 51539177

diff --git a/Engine/Source/Developer/SettingsEditor/Private/Customizations/SettingsEditorCustomization.cpp b/Engine/Source/Developer/SettingsEditor/Private/Customizations/SettingsEditorCustomization.cpp
--- a/Engine/Source/Developer/SettingsEditor/Private/Customizations/SettingsEditorCustomization.cpp
+++ b/Engine/Source/Developer/SettingsEditor/Private/Customizations/SettingsEditorCustomization.cpp
@@ -5,70 +5,51 @@
 #include "DetailLayoutBuilder.h"
 #include "DetailWidgetRow.h"
 #include "ISettingsEditorModule.h"
+#include "ISettingsModule.h"
+#include "ISettingsSection.h"
 #include "HAL/FileManager.h"
 #include "Misc/ConfigContext.h"
 #include "Modules/ModuleManager.h"
+#include "SettingsEditorLogs.h"
 #include "UObject/Package.h"
+#include "UObject/TextProperty.h"
 
 #define LOCTEXT_NAMESPACE "SettingsEditorCustomization"
 
-TSharedRef<IDetailCustomization> FSettingsEditorCustomization::MakeInstance()
-{
-	return MakeShared<FSettingsEditorCustomization>();
-}
+TMap<FString, TMap<FName, FString>> FSettingsEditorCustomization::ObjectDefaultPropertyValues;
 
-namespace UE::Editor::Settings::Private
-{
-bool GetValueFromObject(TSharedPtr<IPropertyHandle> PropertyHandle, UObject* ValueSource, FString& OutValue)
+namespace UE::SettingsEditor::Private
 {
-	const FProperty* SettingsProperty = PropertyHandle->GetProperty();
-	const void* PropertyAddress = PropertyHandle->GetValueBaseAddress((uint8*)ValueSource);
-
-	constexpr uint32 PortFlags = PPF_None;
-	if (PropertyHandle->GetArrayIndex() == INDEX_NONE && SettingsProperty->ArrayDim > 1)
-	{
-		FArrayProperty::ExportTextInnerItem(OutValue, SettingsProperty, PropertyAddress, SettingsProperty->ArrayDim,
-			PropertyAddress, SettingsProperty->ArrayDim, nullptr, PortFlags);
-	}
-	else
-	{
-		PropertyHandle->GetProperty()->ExportTextItem_Direct(OutValue, PropertyAddress, PropertyAddress, nullptr, PortFlags);
-	}
-
-	const UEnum* Enum = nullptr;
-	int64 EnumValue = 0;
-	if (const FByteProperty* ByteProperty = CastField<FByteProperty>(SettingsProperty))
-	{
-		if (ByteProperty->Enum != nullptr)
+	static FAutoConsoleCommand ClearObjectDefaultPropertyValuesCache(
+		TEXT("SettingsEditor.ClearObjectDefaultPropertyValuesCache"),
+		TEXT("Clears the SettingsEditor objects default values cache to force it to rebuild."),
+		FConsoleCommandDelegate::CreateStatic([]()
 		{
-			Enum = ByteProperty->Enum.Get();
-			EnumValue = ByteProperty->GetPropertyValue(PropertyAddress);
-		}
-	}
-	else if (const FEnumProperty* EnumProperty = CastField<FEnumProperty>(SettingsProperty))
-	{
-		Enum = EnumProperty->GetEnum();
-		EnumValue = EnumProperty->GetUnderlyingProperty()->GetSignedIntPropertyValue(PropertyAddress);
-	}
-
-	if (Enum)
-	{
-		if (Enum->IsValidEnumValue(EnumValue))
-		{
-			OutValue = Enum->GetDisplayNameTextByValue(EnumValue).ToString();
-			if (OutValue.IsEmpty())
+			const FPropertyEditorModule& PropertyEditorModule = FModuleManager::GetModuleChecked<FPropertyEditorModule>("PropertyEditor");
+			const ISettingsModule& SettingsModule = FModuleManager::LoadModuleChecked<ISettingsModule>(TEXT("Settings"));
+			TArray<FName> SettingsDetailsViewIds;
+			SettingsModule.GetContainerNames(SettingsDetailsViewIds);
+			// This will force cache to be rebuilt on next refresh which we will force below
+			FSettingsEditorCustomization::ClearObjectDefaultPropertyValues();
+			for (const FName& SettingsDetailsViewId : SettingsDetailsViewIds)
 			{
-				OutValue = Enum->GetNameStringByValue(EnumValue);
+				if (const TSharedPtr<IDetailsView> EditorDetailsView = PropertyEditorModule.FindDetailView(SettingsDetailsViewId))
+				{
+					EditorDetailsView->RequestForceRefresh();
+				}
 			}
-		}
-		else
-		{
-			return false;
-		}
-	}
+		})
+	);
+}
 
-	return true;
+TSharedRef<IDetailCustomization> FSettingsEditorCustomization::MakeInstance()
+{
+	return MakeShared<FSettingsEditorCustomization>();
 }
+
+void FSettingsEditorCustomization::ClearObjectDefaultPropertyValues()
+{
+	ObjectDefaultPropertyValues.Reset();
 }
 
 void FSettingsEditorCustomization::CustomizeDetails(IDetailLayoutBuilder& InDetailLayoutBuilder)
@@ -109,8 +90,10 @@ void FSettingsEditorCustomization::CustomizeDetails(IDetailLayoutBuilder& InDeta
 			continue;
 		}
 
+		CacheResetDefaultPropertyValues(InDetailLayoutBuilder, PinnedSettingsObject.Get());
+
 		const UClass* SettingsClass = SettingsObject->GetClass();
-		for (const FProperty* SettingsProperty = SettingsClass->PropertyLink; SettingsProperty; SettingsProperty = SettingsProperty->PropertyLinkNext)
+		for (const FProperty* SettingsProperty : TFieldRange<FProperty>(SettingsClass, EFieldIteratorFlags::IncludeSuper))
 		{
 			if (!SettingsProperty->HasAnyPropertyFlags(CPF_Config | CPF_GlobalConfig))
 			{
@@ -129,6 +112,14 @@ void FSettingsEditorCustomization::CustomizeDetails(IDetailLayoutBuilder& InDeta
 				continue;
 			}
 			
+			// Override reset to default
+			PropertyRow->OverrideResetToDefault(
+				FResetToDefaultOverride::Create(
+					FIsResetToDefaultVisible::CreateSP(this, &FSettingsEditorCustomization::IsResetToDefaultVisible),
+					FResetToDefaultHandler::CreateSP(this, &FSettingsEditorCustomization::OnResetToDefault)
+				)
+			);
+
 			const FString& CVar = SettingsProperty->GetMetaData(TEXT("ConsoleVariable"));
 			if (CVarSearchTerm != nullptr && !CVar.IsEmpty())
 			{
@@ -143,12 +134,6 @@ void FSettingsEditorCustomization::CustomizeDetails(IDetailLayoutBuilder& InDeta
 					PropertyHandle->CreatePropertyValueWidget()
 				];
 
-				// Override reset to default
-                Row.OverrideResetToDefault(FResetToDefaultOverride::Create(
-                	FIsResetToDefaultVisible::CreateSP(this, &FSettingsEditorCustomization::IsResetToDefaultVisible),
-                	FResetToDefaultHandler::CreateSP(this, &FSettingsEditorCustomization::OnResetToDefault))
-                );
-
 				// Allow searching of cvar in settings editor (for i.e. cvar="r.vsync")
 				FText FilterText = FText::FormatOrdered(TokenDelimiter, FText::FromString(CVar));
 
@@ -168,40 +153,26 @@ void FSettingsEditorCustomization::CustomizeDetails(IDetailLayoutBuilder& InDeta
 				FilterText = FText::FormatOrdered(TokenCombineFormat, SearchKeyCVar.ToLower(), FilterText);
 				Row.FilterString(FilterText);
 			}
-			else
-			{
-				PropertyRow->OverrideResetToDefault(FResetToDefaultOverride::Create(
-					FIsResetToDefaultVisible::CreateSP(this, &FSettingsEditorCustomization::IsResetToDefaultVisible),
-					FResetToDefaultHandler::CreateSP(this, &FSettingsEditorCustomization::OnResetToDefault))
-				);
-			}
 		}
 	}
 }
 
-void FSettingsEditorCustomization::CustomizeDetails(const TSharedPtr<IDetailLayoutBuilder>& DetailBuilder)
-{
-	check(DetailBuilder)
-	DetailBuilderPtr = DetailBuilder;
-	FSettingsEditorCustomization::CustomizeDetails(*DetailBuilder);
-}
-
 void FSettingsEditorCustomization::OnResetToDefault(TSharedPtr<IPropertyHandle> InHandle)
 {
-	const FProperty* SettingsProperty = InHandle->GetProperty();
-	if (!SettingsProperty)
+	TArray<UObject*> SettingsOuterObjects;
+	InHandle->GetOuterObjects(SettingsOuterObjects);
+	if (SettingsOuterObjects.IsEmpty() || !SettingsOuterObjects[0])
 	{
 		return;
 	}
 
-	TArray<UObject*> SettingsOuterObjects;
-	InHandle->GetOuterObjects(SettingsOuterObjects);
-	if (SettingsOuterObjects.IsEmpty())
+	TMap<FName, FString>* PropertyValues = ObjectDefaultPropertyValues.Find(SettingsOuterObjects[0]->GetPathName());
+	if (!PropertyValues)
 	{
 		return;
 	}
 
-	const FString* DefaultValuePtr = DefaultPropertyValues.Find(InHandle->GetProperty()->GetFName());
+	const FString* DefaultValuePtr = PropertyValues->Find(InHandle->GetProperty()->GetFName());
 	if (!DefaultValuePtr)
 	{
 		return;
@@ -215,13 +186,19 @@ bool FSettingsEditorCustomization::IsResetToDefaultVisible(TSharedPtr<IPropertyH
 {
 	TArray<UObject*> SettingsOuterObjects;
 	InHandle->GetOuterObjects(SettingsOuterObjects);
-	if (SettingsOuterObjects.IsEmpty())
+	if (SettingsOuterObjects.IsEmpty() || !SettingsOuterObjects[0])
+	{
+		return false;
+	}
+
+	TMap<FName, FString>* PropertyValues = ObjectDefaultPropertyValues.Find(SettingsOuterObjects[0]->GetPathName());
+	if (!PropertyValues)
 	{
 		return false;
 	}
 
-	const FString* DefaultValue = DefaultPropertyValues.Find(InHandle->GetProperty()->GetFName());
-	if (!DefaultValue)
+	const FString* DefaultValuePtr = PropertyValues->Find(InHandle->GetProperty()->GetFName());
+	if (!DefaultValuePtr)
 	{
 		return false;
 	}
@@ -233,117 +210,308 @@ bool FSettingsEditorCustomization::IsResetToDefaultVisible(TSharedPtr<IPropertyH
 		return false;
 	}
 
-	return !CurrentValue.Equals(*DefaultValue, ESearchCase::CaseSensitive);
+	return !CurrentValue.Equals(*DefaultValuePtr, ESearchCase::CaseSensitive);
 }
 
-void FSettingsEditorCustomization::ResetDefaultPropertyValues(IDetailLayoutBuilder& InDetailLayoutBuilder, UObject* SettingsObject)
+void FSettingsEditorCustomization::CacheResetDefaultPropertyValues(const IDetailLayoutBuilder& InDetailLayoutBuilder, UObject* InSettingsObject)
 {
-	DefaultPropertyValues.Reset();
+	if (!GConfig 
+		|| !InSettingsObject 
+		|| ObjectDefaultPropertyValues.Contains(InSettingsObject->GetPathName()))
+	{
+		// Done
+		return;
+	}
 
+	TRACE_CPUPROFILER_EVENT_SCOPE(FSettingsEditorCustomization::CacheResetDefaultPropertyValues);
+
+	const UClass* SettingsClass = InSettingsObject->GetClass();
+	const bool bIsDefaultConfig = SettingsClass->HasAnyClassFlags(CLASS_DefaultConfig);
+
+	const FString ConfigFilePath = bIsDefaultConfig
+		? SettingsClass->GetDefaultObject()->GetDefaultConfigFilename()
+		: SettingsClass->GetConfigName();
+
+	FString SectionName = SettingsClass->GetPathName();
+	// Allows to override the config section name for backward compatibility
+	InSettingsObject->OverrideConfigSection(SectionName);
+
+	// Assume the current settings already match the defaults
+	UObject* DefaultObject = InSettingsObject;
+	
+	// Cleanup temp object if any
 	UObject* TempDefaultObject = nullptr;
 	ON_SCOPE_EXIT
 	{
 		if (TempDefaultObject)
 		{
-			TempDefaultObject->ClearFlags(RF_ClassDefaultObject);
-			TempDefaultObject->Rename(nullptr, nullptr, REN_DoNotDirty | REN_DontCreateRedirectors | REN_NonTransactional);
 			TempDefaultObject->MarkAsGarbage();
 			TempDefaultObject = nullptr;
 		}
 	};
 
-	const UClass* SettingsClass = SettingsObject->GetClass();
-	const FString ConfigFilePath = SettingsClass->GetConfigName();
-	const FString SectionName = SettingsClass->GetPathName();
+	// Skip creating temp object for abstract classes but still proceed further
+	if (!SettingsClass->HasAnyClassFlags(CLASS_Abstract))
+	{
+		// If the section exists in the current config, then defaults have been overriden
+		const FConfigFile* ConfigFile = GConfig->Find(ConfigFilePath);
+		if (const FConfigSection* ConfigSection = ConfigFile ? ConfigFile->FindSection(SectionName) : nullptr)
+		{
+			// Allocate a new temporary object to access defaults
+			TempDefaultObject = StaticAllocateObject(SettingsClass, GetTransientPackage(), NAME_None, RF_Transient);
+
+			EObjectInitializerOptions InitOptions = EObjectInitializerOptions::None;
+			if (!SettingsClass->HasAnyClassFlags(CLASS_Native | CLASS_Intrinsic))
+			{
+				// Blueprint CDOs have their properties always initialized.
+				InitOptions |= EObjectInitializerOptions::InitializeProperties;
+			}
+
+			// Initialize temporary object by loading the default config values
+			(*SettingsClass->ClassConstructor)(FObjectInitializer(TempDefaultObject, nullptr, InitOptions));
 
-	// First assume the current settings already match the defaults
-	UObject* ValueSource = SettingsObject;
+			// Look into this temp object next
+			DefaultObject = TempDefaultObject;
+		}
+	}
 
-	// If the section exists in the current config, then defaults may have been overriden
-	const FConfigFile* ConfigFile = GConfig->Find(ConfigFilePath);
-	if (ConfigFile && ConfigFile->FindSection(SectionName))
+	// Get engine base config file eg: BaseInput.ini, BaseEngine.ini...
+	const FString ConfigName = FPaths::GetBaseFilename(SettingsClass->GetConfigName());
+	const FString BaseIniPath = FConfigCacheIni::NormalizeConfigIniPath((FPaths::EngineConfigDir() / FString::Printf(TEXT("Base%s.ini"), *ConfigName)));
+
+	// Load base config from buffer but handle special symbol for arrays/sets...
+	FConfigFile ConfigFile;
+	if (FConfigFile* BaseConfigFile = GConfig->Find(BaseIniPath))
+	{
+		FString Buffer;
+		BaseConfigFile->WriteToString(Buffer);
+		ConfigFile.CombineFromBuffer(Buffer, FString(), /** HandleSymbolCommands */true);
+	}
+
+	// Add entry to avoid processing same object next time
+	TMap<FName, FString>& PropertyValues = ObjectDefaultPropertyValues.FindOrAdd(InSettingsObject->GetPathName());
+	for (FProperty* SettingsProperty : TFieldRange<FProperty>(SettingsClass, EFieldIteratorFlags::IncludeSuper))
 	{
-		constexpr bool bIncludeDotInExtension = true;
-		const FString ConfigDirectory = FPaths::GetPath(ConfigFilePath);
-		const FString BaseIniName = FPaths::GetBaseFilename(ConfigFilePath);
-		const FString ConfigExtension = FPaths::GetExtension(ConfigFilePath, bIncludeDotInExtension);
-
-		uint32 Counter = 1;
-		FString BackupConfigPath = FPaths::Combine(ConfigDirectory, FString::Printf(TEXT("%s_Backup%s"), *BaseIniName, *ConfigExtension));
-		while (IFileManager::Get().FileExists(*BackupConfigPath))
+		if (!SettingsProperty->HasAnyPropertyFlags(CPF_Config | CPF_GlobalConfig))
 		{
-			BackupConfigPath = FPaths::Combine(ConfigDirectory, FString::Printf(TEXT("%s_Backup%d%s"), *BaseIniName, Counter++, *ConfigExtension));
+			continue;
 		}
 
-		// Backup the current config in case something were to go wrong in this process
-		if (IFileManager::Get().Copy(*BackupConfigPath, *ConfigFilePath) != ECopyResult::COPY_OK)
+		TSharedRef<IPropertyHandle> PropertyHandle = InDetailLayoutBuilder.GetProperty(SettingsProperty->GetFName(), SettingsClass);
+		if (!PropertyHandle->IsValidHandle())
 		{
-			return;
+			continue;
 		}
+		
+		FString DefaultValue;
+		FString DefaultValueSource = TEXT("Class Default Object");
+		bool bFound = false;
+
+		// Look for cvar default value...
+		const FString& ConsoleVar = PropertyHandle->GetMetaData(TEXT("ConsoleVariable"));
+		if (!ConsoleVar.IsEmpty())
+		{
+			if (IConsoleVariable* CVar = IConsoleManager::Get().FindConsoleVariable(*ConsoleVar))
+			{
+				DefaultValue = CVar->GetDefaultValue();
 
-		// Get rid of the current config values that may be overriding the defaults
-		GConfig->EmptySection(*SectionName, ConfigFilePath);
-		GConfig->Flush(false, ConfigFilePath);
+				const UEnum* Enum = nullptr;
+				int64 EnumValue;
+				LexFromString(EnumValue, *DefaultValue);
+				
+				// Sanitize Bool : "true" => "True", "false" => "False" since detection is case-sensitive
+				if (const FBoolProperty* BoolProperty = CastField<FBoolProperty>(SettingsProperty))
+				{
+					DefaultValue = DefaultValue.ToBool() ? FCoreTexts::Get().True.ToString() : FCoreTexts::Get().False.ToString();
+					bFound = true;
+				}
+				// Extract string value from enum : "0" becomes "FirstValueName"
+				else if (const FByteProperty* ByteProperty = CastField<FByteProperty>(SettingsProperty))
+				{
+					Enum = ByteProperty->Enum.Get();
+				}
+				else if (const FEnumProperty* EnumProperty = CastField<FEnumProperty>(SettingsProperty))
+				{
+					Enum = EnumProperty->GetEnum();
+				}
+				// Property is a float/string/int or a type that does not need additional processing
+				else
+				{
+					bFound = true;
+				}
 
-		// Reload the config, ensuring that we're only looking at default values
-		FConfigContext::ForceReloadIntoGConfig().Load(*BaseIniName);
+				if (Enum && Enum->IsValidEnumValue(EnumValue))
+				{
+					DefaultValue = Enum->GetNameStringByValue(EnumValue);
+					bFound = true;
+				}
+				
+				if (bFound)
+				{
+					DefaultValueSource = FString::Printf(TEXT("CVAR=%s"), *ConsoleVar);
+				}
+			}
+		}
 
-		UClass* ParentClass = SettingsClass->GetSuperClass();
-		UObject* ParentDefaultObject = nullptr;
-		if (ParentClass)
+		// Look into base config file...
+		if (!bFound)
 		{
-			UObjectForceRegistration(ParentClass);
-			ParentDefaultObject = ParentClass->GetDefaultObject();
-			check(GConfig);
-		}
+			const FString ConfigKey = SettingsProperty->GetName();
+			const FConfigSection* ConfigSection = ConfigFile.FindSection(SectionName);
+			if (ConfigSection && ConfigSection->Contains(*ConfigKey))
+			{
+				bFound = true;
+				
+				// Check if we need to retrieve a single or multiple values (array, set)
+				if (SettingsProperty->IsA<FArrayProperty>() 
+					|| SettingsProperty->IsA<FSetProperty>())
+				{
+					FProperty* InnerProp = nullptr;
+					if (const FSetProperty* SetProperty = CastField<FSetProperty>(SettingsProperty))
+					{
+						InnerProp = SetProperty->ElementProp;
+					}
+					else if (const FArrayProperty* ArrayProperty = CastField<FArrayProperty>(SettingsProperty))
+					{
+						InnerProp = ArrayProperty->Inner;
+					}
 
-		// Allocate a new, temporary default object
-		TempDefaultObject = StaticAllocateObject(SettingsClass, GetTransientPackage(), NAME_None, RF_ClassDefaultObject);
-		check(SettingsClass->GetDefaultObject() != TempDefaultObject); // Ensure we haven't been given the existing default object
+					TArray<FString> DefaultValues;
+					ConfigFile.GetArray(
+						*SectionName,
+						*ConfigKey,
+						DefaultValues
+					);
+
+					// Wrap with "quotes" for string otherwise format is not valid to compare import vs export, eg: ("/Game","/Engine") vs (/Game, /Engine)
+					if (InnerProp && (InnerProp->IsA<FStrProperty>() 
+						|| InnerProp->IsA<FNameProperty>() 
+						|| InnerProp->IsA<FTextProperty>()))
+					{
+						for (FString& Value : DefaultValues)
+						{
+							Value = FString::Printf(TEXT("\"%s\""), *Value);
+						}
+					}
 
-		EObjectInitializerOptions InitOptions = EObjectInitializerOptions::None;
-		if (!SettingsClass->HasAnyClassFlags(CLASS_Native | CLASS_Intrinsic))
-		{
-			// Blueprint CDOs have their properties always initialized.
-			InitOptions |= EObjectInitializerOptions::InitializeProperties;
-		}
+					// Wrap in array (...) otherwise format is not valid to compare import vs export
+					DefaultValue = FString::Printf(TEXT("(%s)"), *FString::Join(DefaultValues, TEXT(",")));
+				}
+				else
+				{
+					// Maps are serialized in a single line eg: MapPropertyName=((Key, "Value"),(Key2, "Value2"))
+					ConfigFile.GetString(
+						*SectionName,
+						*ConfigKey,
+						DefaultValue
+					);
+				}
+
+				// We need to normalize values to avoid precision and format issues 
+				// eg: Current Value (R=0.594000,G=0.019700,B=0.000000,A=1.000000) VS Config Value (R=0.594,G=0.0197,B=0.0,A=1.0) are not equal
+				FString Normalized;
+
+				// Allocate a buffer that satisfies the property size and alignment
+				const int32 PropSize = SettingsProperty->GetSize();
+				const int32 PropAlignment = SettingsProperty->GetMinAlignment();
+				if (void* TempBuffer = FMemory::Malloc(PropSize, PropAlignment))
+				{
+					SettingsProperty->InitializeValue(TempBuffer);
 
-		// Initialize the temporary default object, loading the default config values
-		(*SettingsClass->ClassConstructor)(FObjectInitializer(TempDefaultObject, ParentDefaultObject, InitOptions));
-		ValueSource = TempDefaultObject;
+					if (const TCHAR* Result = SettingsProperty->ImportText_Direct(*DefaultValue, TempBuffer, nullptr, PPF_None))
+					{
+						// Export back through the same path so both strings are formatted identically...
+						SettingsProperty->ExportTextItem_Direct(Normalized, TempBuffer, nullptr, nullptr, PPF_None);
+					}
+					else
+					{
+						Normalized = DefaultValue;
+					}
+
+					SettingsProperty->DestroyValue(TempBuffer);
+					FMemory::Free(TempBuffer);
+				}
+				else
+				{
+					Normalized = DefaultValue;
+				}
+
+				DefaultValue = Normalized;
+				DefaultValueSource = FString::Printf(TEXT("INI=%s[%s][%s]"), *BaseIniPath, *SectionName, *SettingsProperty->GetName());	
+			}
+		}
 
-		// Restore the current config values
-		if (IFileManager::Get().Copy(*ConfigFilePath, *BackupConfigPath) != ECopyResult::COPY_OK
-			|| !FConfigContext::ForceReloadIntoGConfig().Load(*BaseIniName))
+		// Look into the default object defined previously...
+		if (!bFound && !GetValueFromObject(PropertyHandle, DefaultObject, DefaultValue))
 		{
-			UE_LOG(LogConfig, Warning, TEXT("Failed to restore config %s while caching default values. Click Import in %s and select %s to restore previous values."), *ConfigFilePath, *SectionName, *BackupConfigPath);
-			return;
+			UE_LOG(LogSettingsEditor, Warning, TEXT("Could not retrieve default value for property %s in %s"), *PropertyHandle->GetPropertyDisplayName().ToString(), *DefaultObject->GetName())
+			continue;
 		}
 
-		IFileManager::Get().Delete(*BackupConfigPath);
+		PropertyValues.Add(SettingsProperty->GetFName(), DefaultValue);
+
+		// Set default value and default value source in tooltip for convenience
+		const FText PropertyTooltip = FText::Format(
+			INVTEXT("{0}\n\nDefault Value: {1}\n\nDefault Value Source: {2}"), 
+			PropertyHandle->GetToolTipText(), 
+			FText::FromString(*DefaultValue),
+			FText::FromString(*DefaultValueSource)
+		);
+
+		PropertyHandle->SetToolTipText(PropertyTooltip);
 	}
+}
 
-	for (const FProperty* SettingsProperty = SettingsClass->PropertyLink; SettingsProperty; SettingsProperty = SettingsProperty->PropertyLinkNext)
+bool FSettingsEditorCustomization::GetValueFromObject(const TSharedPtr<IPropertyHandle>& InPropertyHandle, UObject* InValueSource, FString& OutValue)
+{
+	if (!InValueSource || !InPropertyHandle || !InPropertyHandle->GetProperty())
 	{
-		if (!SettingsProperty->HasAnyPropertyFlags(CPF_Config | CPF_GlobalConfig))
+		return false;
+	}
+
+	const FProperty* SettingsProperty = InPropertyHandle->GetProperty();
+	const void* PropertyAddress = InPropertyHandle->GetValueBaseAddress(reinterpret_cast<uint8*>(InValueSource));
+
+	constexpr uint32 PortFlags = PPF_None;
+	if (InPropertyHandle->GetArrayIndex() == INDEX_NONE && SettingsProperty->ArrayDim > 1)
+	{
+		FArrayProperty::ExportTextInnerItem(OutValue, SettingsProperty, PropertyAddress, SettingsProperty->ArrayDim,
+			PropertyAddress, SettingsProperty->ArrayDim, nullptr, PortFlags);
+	}
+	else
+	{
+		SettingsProperty->ExportTextItem_Direct(OutValue, PropertyAddress, PropertyAddress, nullptr, PortFlags);
+	}
+
+	const UEnum* Enum = nullptr;
+	int64 EnumValue = 0;
+	if (const FByteProperty* ByteProperty = CastField<FByteProperty>(SettingsProperty))
+	{
+		if (ByteProperty->Enum != nullptr)
 		{
-			continue;
+			Enum = ByteProperty->Enum.Get();
+			EnumValue = ByteProperty->GetPropertyValue(PropertyAddress);
 		}
+	}
+	else if (const FEnumProperty* EnumProperty = CastField<FEnumProperty>(SettingsProperty))
+	{
+		Enum = EnumProperty->GetEnum();
+		EnumValue = EnumProperty->GetUnderlyingProperty()->GetSignedIntPropertyValue(PropertyAddress);
+	}
 
-		TSharedRef<IPropertyHandle> PropertyHandle = InDetailLayoutBuilder.GetProperty(SettingsProperty->GetFName(), SettingsClass);
-		if (!PropertyHandle->IsValidHandle())
+	if (Enum)
+	{
+		if (Enum->IsValidEnumValue(EnumValue))
 		{
-			continue;
+			OutValue = Enum->GetNameStringByValue(EnumValue);
 		}
-
-		FString DefaultValue;
-		if (!UE::Editor::Settings::Private::GetValueFromObject(PropertyHandle, ValueSource, DefaultValue))
+		else
 		{
-			continue;
+			return false;
 		}
-
-		DefaultPropertyValues.Add(PropertyHandle->GetProperty()->GetFName(), DefaultValue);
 	}
+
+	return true;
 }
 
 #undef LOCTEXT_NAMESPACE
diff --git a/Engine/Source/Developer/SettingsEditor/Private/Customizations/SettingsEditorCustomization.h b/Engine/Source/Developer/SettingsEditor/Private/Customizations/SettingsEditorCustomization.h
--- a/Engine/Source/Developer/SettingsEditor/Private/Customizations/SettingsEditorCustomization.h
+++ b/Engine/Source/Developer/SettingsEditor/Private/Customizations/SettingsEditorCustomization.h
@@ -14,16 +14,19 @@ class FSettingsEditorCustomization : public IDetailCustomization
 {
 public:
 	static TSharedRef<IDetailCustomization> MakeInstance();
+	static void ClearObjectDefaultPropertyValues();
+
+	//~ Begin IDetailCustomization
 	virtual void CustomizeDetails(IDetailLayoutBuilder& InDetailLayoutBuilder) override;
-	virtual void CustomizeDetails(const TSharedPtr<IDetailLayoutBuilder>& DetailBuilder) override;
+	//~ End IDetailCustomization
 
 private:
+	/** Stores all settings object default property values ObjectPath->PropertyName->Value */
+	static TMap<FString, TMap<FName, FString>> ObjectDefaultPropertyValues;
+	static bool GetValueFromObject(const TSharedPtr<IPropertyHandle>& InPropertyHandle, UObject* InValueSource, FString& OutValue);
+
+	/** Saves once at startup all settings object default property values */
+	void CacheResetDefaultPropertyValues(const IDetailLayoutBuilder& InDetailLayoutBuilder, UObject* InSettingsObject);
 	void OnResetToDefault(TSharedPtr<IPropertyHandle> InHandle);
 	bool IsResetToDefaultVisible(TSharedPtr<IPropertyHandle> InHandle);
-
-	// Cache the latest default property values
-	void ResetDefaultPropertyValues(IDetailLayoutBuilder& InDetailLayoutBuilder, UObject* SettingsObject);
-
-	TMap<FName, FString> DefaultPropertyValues;
-	TWeakPtr<IDetailLayoutBuilder> DetailBuilderPtr;
 };
diff --git a/Engine/Source/Developer/SettingsEditor/Private/SettingsEditorLogs.h b/Engine/Source/Developer/SettingsEditor/Private/SettingsEditorLogs.h
new file mode 100644
--- /dev/null
+++ b/Engine/Source/Developer/SettingsEditor/Private/SettingsEditorLogs.h
@@ -0,0 +1,7 @@
+// Copyright Epic Games, Inc. All Rights Reserved.
+
+#pragma once
+
+#include "Logging/LogMacros.h"
+
+DECLARE_LOG_CATEGORY_EXTERN(LogSettingsEditor, Log, All);
\ No newline at end of file
diff --git a/Engine/Source/Developer/SettingsEditor/Private/SettingsEditorModule.cpp b/Engine/Source/Developer/SettingsEditor/Private/SettingsEditorModule.cpp
--- a/Engine/Source/Developer/SettingsEditor/Private/SettingsEditorModule.cpp
+++ b/Engine/Source/Developer/SettingsEditor/Private/SettingsEditorModule.cpp
@@ -15,6 +15,7 @@
 #include "PropertyEditorClipboard.h"
 #include "PropertyEditorDelegates.h"
 #include "PropertyEditorModule.h"
+#include "SettingsEditorLogs.h"
 #include "ToolMenus.h"
 #include "UObject/UObjectHash.h"
 #include "UObject/UObjectIterator.h"
@@ -24,6 +25,8 @@
 #include "Widgets/Notifications/SNotificationList.h"
 #include "Widgets/SSettingsEditor.h"
 
+DEFINE_LOG_CATEGORY(LogSettingsEditor);
+
 #define LOCTEXT_NAMESPACE "SSettingsEditor"
 
 /** Holds auto discovered settings information so that they can be unloaded automatically when refreshing. */
diff --git a/Engine/Source/Editor/InputBindingEditor/Private/InputBindingEditorModule.cpp b/Engine/Source/Editor/InputBindingEditor/Private/InputBindingEditorModule.cpp
--- a/Engine/Source/Editor/InputBindingEditor/Private/InputBindingEditorModule.cpp
+++ b/Engine/Source/Editor/InputBindingEditor/Private/InputBindingEditorModule.cpp
@@ -416,6 +416,16 @@ public:
 							CommandInfoTooltip = FText::Format(LOCTEXT("CommandInfoDebugToolTip", "{0}\n\nBinding Context: {1}\nCommand Name: {2}"), CommandInfoTooltip, FText::FromName(CommandInfo->GetBindingContext()), FText::FromName(CommandInfo->GetCommandName()));
 						}
 
+						// Show default chords	
+						const FInputChord& DefaultChord1 = CommandInfo->GetDefaultChord(EMultipleKeyBindingIndex::Primary);
+						const FInputChord& DefaultChord2 = CommandInfo->GetDefaultChord(EMultipleKeyBindingIndex::Secondary);
+						CommandInfoTooltip = FText::Format(
+							INVTEXT("{0}\n\nDefault Primary: \"{1}\"\nDefault Secondary: \"{2}\""),
+							CommandInfoTooltip,
+							DefaultChord1.GetInputText(),
+							DefaultChord2.GetInputText()
+						);
+
 						return CommandInfoTooltip;
 					})
 				];
@@ -439,10 +449,55 @@ public:
 						SNew(SChordEditBox, CommandInfo, EMultipleKeyBindingIndex::Secondary)
 					]
 				];
+				
+				Row->OverrideResetToDefault(
+					FResetToDefaultOverride::Create(
+						TAttribute<bool>::CreateSP(this, &FEditorKeyboardShortcutSettings::IsCommandResetToDefaultVisible, CommandInfo),
+						FSimpleDelegate::CreateSP(this, &FEditorKeyboardShortcutSettings::ResetToDefaultCommand, CommandInfo)
+					)
+				);
 			}
 		}
 	}
 
+	bool IsCommandResetToDefaultVisible(TSharedPtr<FUICommandInfo> InCommandInfo) const
+	{
+		if (!InCommandInfo)
+		{
+			return false;
+		}
+
+		const FInputChord& DefaultChord1 = InCommandInfo->GetDefaultChord(EMultipleKeyBindingIndex::Primary);
+		const TSharedRef<const FInputChord> ActiveChord1 = InCommandInfo->GetActiveChord(EMultipleKeyBindingIndex::Primary);
+		if (ActiveChord1.Get() != DefaultChord1)
+		{
+			return true;
+		}
+
+		const FInputChord& DefaultChord2 = InCommandInfo->GetDefaultChord(EMultipleKeyBindingIndex::Secondary);
+		const TSharedRef<const FInputChord> ActiveChord2 = InCommandInfo->GetActiveChord(EMultipleKeyBindingIndex::Secondary);
+		if (ActiveChord2.Get() != DefaultChord2)
+		{
+			return true;
+		}
+
+		return false;
+	}
+
+	void ResetToDefaultCommand(TSharedPtr<FUICommandInfo> InCommandInfo)
+	{
+		if (!InCommandInfo)
+		{
+			return;
+		}
+
+		const FInputChord& DefaultChord1 = InCommandInfo->GetDefaultChord(EMultipleKeyBindingIndex::Primary);
+		InCommandInfo->SetActiveChord(DefaultChord1, EMultipleKeyBindingIndex::Primary);
+
+		const FInputChord& DefaultChord2 = InCommandInfo->GetDefaultChord(EMultipleKeyBindingIndex::Secondary);
+		InCommandInfo->SetActiveChord(DefaultChord2, EMultipleKeyBindingIndex::Secondary);
+	}
+
 private:
 	bool bUpdateRequested;
 	IDetailLayoutBuilder* DetailBuilder;
