EditorUsability : SettingsEditor
- Fixed reported crash when temporary default object is created without RF_ClassDefaultObject flag in NavigationSystem.cpp:940

#Synced-CL 51539177

diff --git a/Engine/Source/Developer/SettingsEditor/Private/Customizations/SettingsEditorCustomization.cpp b/Engine/Source/Developer/SettingsEditor/Private/Customizations/SettingsEditorCustomization.cpp
--- a/Engine/Source/Developer/SettingsEditor/Private/Customizations/SettingsEditorCustomization.cpp
+++ b/Engine/Source/Developer/SettingsEditor/Private/Customizations/SettingsEditorCustomization.cpp
@@ -258,7 +258,7 @@ void FSettingsEditorCustomization::CacheResetDefaultPropertyValues(const IDetail
 		if (const FConfigSection* ConfigSection = ConfigFile ? ConfigFile->FindSection(SectionName) : nullptr)
 		{
 			// Allocate a new temporary object to access defaults
-			TempDefaultObject = StaticAllocateObject(SettingsClass, GetTransientPackage(), NAME_None, RF_Transient);
+			TempDefaultObject = StaticAllocateObject(SettingsClass, GetTransientPackage(), NAME_None, RF_ClassDefaultObject | RF_Transient);
 
 			EObjectInitializerOptions InitOptions = EObjectInitializerOptions::None;
 			if (!SettingsClass->HasAnyClassFlags(CLASS_Native | CLASS_Intrinsic))
