EditorUsability : InputBinding
- Added new input settings tab containing keyboard shortcuts and future shortcuts
- Added sidebar to settings editor that shows left panel when more than one item (docked) and hides when it contains a single item
- Import/export (All) will only handle shortcuts instead of all editor settings now
- Updated autocompletion and suggestions to use new container
- Updated permissions to allow opening this tab

#Synced-CL 52634927

diff --git a/Engine/Source/Developer/SettingsEditor/Private/Customizations/SettingsEditorCustomization.cpp b/Engine/Source/Developer/SettingsEditor/Private/Customizations/SettingsEditorCustomization.cpp
--- a/Engine/Source/Developer/SettingsEditor/Private/Customizations/SettingsEditorCustomization.cpp
+++ b/Engine/Source/Developer/SettingsEditor/Private/Customizations/SettingsEditorCustomization.cpp
@@ -76,8 +76,8 @@ void FSettingsEditorCustomization::CustomizeDetails(IDetailLayoutBuilder& InDeta
 				CVarSearchTerm->IconStyleName = "DebugConsole.Icon";
 				CVarSearchTerm->DisallowedSections.Add(TEXT("InputBindings")); // Keyboard Shortcuts
 				CVarSearchTerm->AllowedSections.Add(NAME_None); // All section
-				// No need to show filter buttons in Editor settings
-				CVarSearchTerm->bShowFilter = !SearchContext.IsEqual("Editor");
+				// Only show filter buttons in Project settings
+				CVarSearchTerm->bShowFilter = SearchContext.IsEqual("Project");
 			}
 		}
 	}
diff --git a/Engine/Source/Developer/SettingsEditor/Private/Widgets/SSettingsEditor.cpp b/Engine/Source/Developer/SettingsEditor/Private/Widgets/SSettingsEditor.cpp
--- a/Engine/Source/Developer/SettingsEditor/Private/Widgets/SSettingsEditor.cpp
+++ b/Engine/Source/Developer/SettingsEditor/Private/Widgets/SSettingsEditor.cpp
@@ -34,14 +34,36 @@
 #include "Misc/MessageDialog.h"
 #include "Misc/TransactionObjectEvent.h"
 #include "Models/SettingsEditorMenuContext.h"
+#include "Sidebar/ISidebarDrawerContent.h"
+#include "Sidebar/SSidebarContainer.h"
 #include "Widgets/Layout/SBox.h"
 
 #define LOCTEXT_NAMESPACE "SSettingsEditor"
 
 FLazyName SSettingsEditor::SettingsEditorToolbarName = TEXT("SettingsEditor.SearchToolbar");
+FLazyName SSettingsEditor::SettingsEditorSidebarId = TEXT("CategoryTree");
 const FTextFormat SSettingsEditor::TokenCombineFormat = INVTEXT("{0} {1}");
 const FTextFormat SSettingsEditor::TokenDelimiter = INVTEXT("\"{0}\"");
 
+/** For the category tree sidebar */
+class FSettingsEditorDrawerContent : public ISidebarDrawerContent
+{
+public:
+	explicit FSettingsEditorDrawerContent(TSharedRef<SWidget> InSidebarContent)
+		: SidebarContent(InSidebarContent)
+	{}
+
+	//~ Begin ISidebarDrawerContent
+	virtual FName GetUniqueId() const override { return SSettingsEditor::SettingsEditorSidebarId; }
+	virtual FName GetSectionId() const override { return TEXT("Settings"); }
+	virtual FText GetSectionDisplayText() const override { return LOCTEXT("SettingsSection", "Settings"); }
+	virtual TSharedRef<SWidget> CreateContentWidget() override { return SidebarContent; }
+	//~ End ISidebarDrawerContent
+
+private:
+	TSharedRef<SWidget> SidebarContent;
+};
+
 /* SSettingsEditor structors
  *****************************************************************************/
 
@@ -123,110 +145,125 @@ void SSettingsEditor::Construct( const FArguments& InArgs, const ISettingsEditor
 		);
 	}
 
-	TSharedPtr<SBorder> CategoryTreeBorder;
+	// Register sidebar
+	{
+		SidebarContainer = SNew(SSidebarContainer);
+		Sidebar = SNew(SSidebar, SidebarContainer.ToSharedRef())
+			.Visibility(EVisibility::SelfHitTestInvisible)
+			.TabLocation(ESidebarTabLocation::Left)
+			.InitialDrawerSize(0.20f)
+			.DisablePin(true)
+			.OnGetContent(FOnGetContent::CreateSPLambda(this, [this]()
+			{
+				return SNew(SVerticalBox)
+					+ SVerticalBox::Slot()
+					.Padding(4.f)
+					.AutoHeight()
+					[
+						SNew(SHorizontalBox)
+						+ SHorizontalBox::Slot()
+						.AutoWidth()
+						.VAlign(VAlign_Center)
+						.Padding(0.0f)
+						[
+							SAssignNew(MenuToolbarBox, SBox)
+							.MinDesiredHeight(24.f)
+							[
+								// Menu toolbar goes here
+								SNullWidget::NullWidget
+							]
+						]
+						+ SHorizontalBox::Slot()
+						.FillWidth(1.0f)
+						.VAlign(VAlign_Center)
+						.Padding(0.f)
+						[
+							SAssignNew(SearchBox, SAssetSearchBox)
+							.SuggestionListHeight(60.f)
+							.ShowSearchHistory(false) // Only suggest, otherwise its confusing
+							.DelayChangeNotificationsWhileTyping(true)
+							.ToolTipText(FText::GetEmpty()) // Do not show any tooltip, it's showing when suggestions are listed...
+							.HintText(this, &SSettingsEditor::GetSearchHintText)
+							.OnTextChanged(this, &SSettingsEditor::OnSearchTextChanged)
+							.OnTextCommitted(this, &SSettingsEditor::OnSearchTextCommitted)
+							.OnAssetSearchBoxSuggestionFilter(this, &SSettingsEditor::OnAssetSearchSuggestionFilter)
+							.OnAssetSearchBoxSuggestionChosen(this, &SSettingsEditor::OnAssetSearchSuggestionChosen)
+							.OnKeyDownHandler(this, &SSettingsEditor::OnSearchKeyDownHandler)
+							.AddMetaData<FTagMetaData>(FTagMetaData(TEXT("SettingsEditorSearchBox")))
+						]
+						+ SHorizontalBox::Slot()
+						.AutoWidth()
+						.VAlign(VAlign_Center)
+						.Padding(8.0f, 0.0f, 0.0f, 0.0f)
+						[
+							// Import button
+							SNew(SButton)
+							.OnClicked(this, &SSettingsEditor::HandleImportButtonClicked)
+							.Text(LOCTEXT("ImportAllButtonText", "Import All..."))
+							.ToolTipText(LOCTEXT("ImportAllButtonTooltip", "Import all editor settings from a file on your computer"))
+						]
+						+ SHorizontalBox::Slot()
+						.AutoWidth()
+						.VAlign(VAlign_Center)
+						.Padding(8.0f, 0.0f)
+						[
+							// Export button
+							SNew(SButton)
+							.OnClicked(this, &SSettingsEditor::HandleExportButtonClicked)
+							.Text(LOCTEXT("ExportAllButtonText", "Export All..."))
+							.ToolTipText(LOCTEXT("ExportAllButtonTooltip", "Export all editor settings to a file on your computer"))
+						]
+					]
+					+ SVerticalBox::Slot()
+					.Padding(4.f)
+					.FillHeight(1.f)
+					[
+						// settings area
+						SNew(SOverlay)
+						+ SOverlay::Slot()
+						[
+							SettingsView.ToSharedRef()
+						]
+						+ SOverlay::Slot()
+						.Expose(CustomWidgetSlot)
+					];
+			}));
+		SidebarContainer->RebuildSidebar(Sidebar.ToSharedRef(), FSidebarState());
+		
+		// Create category tree later so everything is initialized when callback is triggered
+		TSharedRef<SWidget> CategoryTreeWidget = SNew(SBorder)
+		  .BorderImage(FAppStyle::Get().GetBrush("SettingsEditor.ListBorder"))
+		  .Padding(8.f)
+		  [
+			  SAssignNew(CategoryTree, SSettingsEditorCategoryTree, SettingsContainer)
+			  .InitialSelection(Model->GetSelectedSection())
+			  .OnSelectionChanged(this, &SSettingsEditor::OnSectionSelectionChanged)
+		  ];
+
+		FSidebarDrawerConfig SidebarDrawerConfig;
+		SidebarDrawerConfig.UniqueId = SettingsEditorSidebarId;
+		SidebarDrawerConfig.ButtonText = LOCTEXT("SidebarButtonText", "Settings");
+		SidebarDrawerConfig.ToolTipText = LOCTEXT("SidebarToolTipText", "Toggle the settings category tree sidebar");
+		SidebarDrawerConfig.Icon = FAppStyle::Get().GetBrush("Icons.Settings");
+		SidebarDrawerConfig.InitialState.bIsDocked = true;
+		Sidebar->RegisterDrawer(MoveTemp(SidebarDrawerConfig));
+		Sidebar->RegisterDrawerSection(
+			SettingsEditorSidebarId,
+			MakeShared<FSettingsEditorDrawerContent>(CategoryTreeWidget)
+		);
+	}
+
 	ChildSlot
 	[
 		SNew(SBorder)
 		.BorderImage(FAppStyle::Get().GetBrush("Brushes.Panel"))
 		.Padding(0.f)
 		[
-			SNew(SSplitter)
-			.Style(FAppStyle::Get(), "SplitterPanel")
-			+ SSplitter::Slot()
-			.Value(.35f)
-			[
-				SAssignNew(CategoryTreeBorder, SBorder)
-				.VAlign(VAlign_Fill)
-				.BorderImage(FAppStyle::Get().GetBrush("SettingsEditor.ListBorder"))
-				.Padding(8.f)
-			]
-			+ SSplitter::Slot()
-			[
-				SNew(SVerticalBox)
-				+ SVerticalBox::Slot()
-				.Padding(4.f)
-				.AutoHeight()
-				[
-					SNew(SHorizontalBox)
-					+ SHorizontalBox::Slot()
-					.AutoWidth()
-					.VAlign(VAlign_Center)
-					.Padding(0.0f)
-					[
-						SAssignNew(MenuToolbarBox, SBox)
-						.MinDesiredHeight(24.f)
-						[
-							// Menu toolbar goes here
-							SNullWidget::NullWidget
-						]
-					]
-					+ SHorizontalBox::Slot()
-					.FillWidth(1.0f)
-					.VAlign(VAlign_Center)
-					.Padding(0.f)
-					[
-						SAssignNew(SearchBox, SAssetSearchBox)
-						.SuggestionListHeight(60.f)
-						.ShowSearchHistory(false) // Only suggest, otherwise its confusing
-						.DelayChangeNotificationsWhileTyping(true)
-						.ToolTipText(FText::GetEmpty()) // Do not show any tooltip, it's showing when suggestions are listed...
-						.HintText(this, &SSettingsEditor::GetSearchHintText)
-						.OnTextChanged(this, &SSettingsEditor::OnSearchTextChanged)
-						.OnTextCommitted(this, &SSettingsEditor::OnSearchTextCommitted)
-						.OnAssetSearchBoxSuggestionFilter(this, &SSettingsEditor::OnAssetSearchSuggestionFilter)
-						.OnAssetSearchBoxSuggestionChosen(this, &SSettingsEditor::OnAssetSearchSuggestionChosen)
-						.OnKeyDownHandler(this, &SSettingsEditor::OnSearchKeyDownHandler)
-						.AddMetaData<FTagMetaData>(FTagMetaData(TEXT("SettingsEditorSearchBox")))
-					]
-					+ SHorizontalBox::Slot()
-					.AutoWidth()
-					.VAlign(VAlign_Center)
-					.Padding(8.0f, 0.0f, 0.0f, 0.0f)
-					[
-						// Import button
-						SNew(SButton)
-						.OnClicked(this, &SSettingsEditor::HandleImportButtonClicked)
-						.Text(LOCTEXT("ImportAllButtonText", "Import All..."))
-						.ToolTipText(LOCTEXT("ImportAllButtonTooltip", "Import all editor settings from a file on your computer"))
-					]
-					+ SHorizontalBox::Slot()
-					.AutoWidth()
-					.VAlign(VAlign_Center)
-					.Padding(8.0f, 0.0f)
-					[
-						// Export button
-						SNew(SButton)
-						.OnClicked(this, &SSettingsEditor::HandleExportButtonClicked)
-						.Text(LOCTEXT("ExportAllButtonText", "Export All..."))
-						.ToolTipText(LOCTEXT("ExportAllButtonTooltip", "Export all editor settings to a file on your computer"))
-					]
-				]
-				+ SVerticalBox::Slot()
-				.Padding(4.f)
-				.FillHeight(1.f)
-				[
-					// settings area
-					SNew(SOverlay)
-					+ SOverlay::Slot()
-					[
-						SettingsView.ToSharedRef()
-					]
-					+ SOverlay::Slot()
-					.Expose(CustomWidgetSlot)
-				]
-			]
+			SidebarContainer.ToSharedRef()
 		]
 	];
 
 	Model->OnSelectionChanged().AddSP(this, &SSettingsEditor::HandleModelSelectionChanged);
-
-	// Create category tree later so everything is initialized when callback is triggered
-	CategoryTreeBorder->SetContent(
-		SAssignNew(CategoryTree, SSettingsEditorCategoryTree, SettingsContainer)
-		.InitialSelection(Model->GetSelectedSection())
-		.OnSelectionChanged(this, &SSettingsEditor::OnSectionSelectionChanged)
-	);
-
 	HandleModelSelectionChanged();
 
 	// Update toolbar menu
@@ -460,6 +497,8 @@ void SSettingsEditor::HandleModelSelectionChanged()
 	// Set all settings object to preview them,
 	// they will get filtered base on the active section selection in the customization
 	{
+		const bool bInit = SettingsView->GetSelectedObjects().IsEmpty();
+
 		TArray<UObject*> SettingsObjects;
 		CategoryTree->ForEachSection([&SettingsObjects](const TSharedPtr<ISettingsSection>& Section)
 			{
@@ -472,6 +511,14 @@ void SSettingsEditor::HandleModelSelectionChanged()
 			}
 		);
 		SettingsView->SetObjects(SettingsObjects);
+		
+		// Settings view will be populated with all settings objects and hide them internally through a customization
+		// Show sidebar if more than 1 settings object
+		const bool bDockSidebar = SettingsObjects.Num() > 1;
+		if (bInit && bDockSidebar != Sidebar->IsDrawerDocked(SettingsEditorSidebarId))
+		{
+			Sidebar->SetDrawerDocked(SettingsEditorSidebarId, bDockSidebar);
+		}
 	}
 }
 
diff --git a/Engine/Source/Developer/SettingsEditor/Private/Widgets/SSettingsEditor.h b/Engine/Source/Developer/SettingsEditor/Private/Widgets/SSettingsEditor.h
--- a/Engine/Source/Developer/SettingsEditor/Private/Widgets/SSettingsEditor.h
+++ b/Engine/Source/Developer/SettingsEditor/Private/Widgets/SSettingsEditor.h
@@ -21,6 +21,8 @@ class SAssetSearchBox;
 class SBox;
 class SScrollBox;
 class SSettingsEditorCategoryTree;
+class SSidebar;
+class SSidebarContainer;
 class USettingsEditorMenuContext;
 class UToolMenu;
 enum class ECheckBoxState : uint8;
@@ -38,6 +40,7 @@ class SSettingsEditor
 public:
 	/** Menu toolbar Identifier */
 	static FLazyName SettingsEditorToolbarName;
+	static FLazyName SettingsEditorSidebarId;
 
 	SLATE_BEGIN_ARGS(SSettingsEditor) { }
 		SLATE_EVENT( FSimpleDelegate, OnApplicationRestartRequired )
@@ -173,4 +176,8 @@ private:
 	 * Note that these objects can be settings objects as well as sub-objects of a settings object.
 	 */
 	TSet<TWeakObjectPtr<UObject>> ModifiedObjectCache;
+
+	/** Sidebar for left panel (category tree) and right panel (details view) */
+	TSharedPtr<SSidebarContainer> SidebarContainer;
+	TSharedPtr<SSidebar> Sidebar;
 };
diff --git a/Engine/Source/Developer/SettingsEditor/Private/Widgets/SSettingsSectionHeader.cpp b/Engine/Source/Developer/SettingsEditor/Private/Widgets/SSettingsSectionHeader.cpp
--- a/Engine/Source/Developer/SettingsEditor/Private/Widgets/SSettingsSectionHeader.cpp
+++ b/Engine/Source/Developer/SettingsEditor/Private/Widgets/SSettingsSectionHeader.cpp
@@ -603,18 +603,7 @@ TSharedPtr<SWidget> FSettingsDetailRootObjectCustomization::CustomizeObjectHeade
 
 bool FSettingsDetailRootObjectCustomization::AreObjectsVisible(const FDetailsObjectSet& InRootObjectSet) const
 {
-	if (SelectedSectionObject == nullptr)
-	{
-		// When viewing "All settings", do not show keyboard shortcut settings...
-		if (InRootObjectSet.RootObjects[0]->GetClass()->GetName() == TEXT("EditorKeyboardShortcutSettings"))
-		{
-			return false;
-		}
-		
-		return true;
-	}
-
-	return SelectedSectionObject == InRootObjectSet.RootObjects[0];
+	return !SelectedSectionObject.IsValid() || (!InRootObjectSet.RootObjects.IsEmpty() && SelectedSectionObject == InRootObjectSet.RootObjects[0]);
 }
 
 bool FSettingsDetailRootObjectCustomization::ShouldDisplayHeader(const FDetailsObjectSet& InRootObjects) const
diff --git a/Engine/Source/Editor/InputBindingEditor/Private/InputBindingEditorModule.cpp b/Engine/Source/Editor/InputBindingEditor/Private/InputBindingEditorModule.cpp
--- a/Engine/Source/Editor/InputBindingEditor/Private/InputBindingEditorModule.cpp
+++ b/Engine/Source/Editor/InputBindingEditor/Private/InputBindingEditorModule.cpp
@@ -254,7 +254,6 @@ public:
 					KeybindSearchTerm->Tooltip = LOCTEXT("SearchKeyKeybindTooltip", "Filters by the pressed key combination");
 					KeybindSearchTerm->IconStyleSetName = FAppStyle::GetAppStyleSetName();
 					KeybindSearchTerm->IconStyleName = TEXT("GraphEditor.KeyEvent_16x");
-					KeybindSearchTerm->AllowedSections.Add(TEXT("InputBindings")); // Only Keyboard Shortcuts
 				}
 
 				if (bContextRegistered)
@@ -264,7 +263,6 @@ public:
 					ContextSearchTerm->Tooltip = LOCTEXT("SearchKeyContextTooltip", "Filters by the command context");
 					// Hide this option from being shown in filter buttons but allow it to be suggested
 					ContextSearchTerm->bShowFilter = false;
-					ContextSearchTerm->AllowedSections.Add(TEXT("InputBindings")); // Only Keyboard Shortcuts
 				}
 			}
 		}
@@ -521,7 +519,7 @@ public:
 		PropertyEditor.RegisterCustomClassLayout(EditorKeyboardShortcutSettingsName, FOnGetDetailCustomizationInstance::CreateStatic(&FEditorKeyboardShortcutSettings::MakeInstance));
 
 		// input bindings
-		const FName ContainerName = "Editor";
+		const FName ContainerName = "InputBinding";
 		const FName CategoryName = "Shortcuts";
 		const ISettingsSectionPtr InputBindingSettingsSection = SettingsModule.RegisterSettings(
 			ContainerName, 
diff --git a/Engine/Source/Editor/InputBindingSettingsViewer/InputBindingSettingsViewer.Build.cs b/Engine/Source/Editor/InputBindingSettingsViewer/InputBindingSettingsViewer.Build.cs
new file mode 100644
--- /dev/null
+++ b/Engine/Source/Editor/InputBindingSettingsViewer/InputBindingSettingsViewer.Build.cs
@@ -0,0 +1,31 @@
+﻿// Copyright Epic Games, Inc. All Rights Reserved.
+
+using UnrealBuildTool;
+
+public class InputBindingSettingsViewer : ModuleRules
+{
+    public InputBindingSettingsViewer(ReadOnlyTargetRules Target) : base(Target)
+    {
+        PCHUsage = ModuleRules.PCHUsageMode.UseExplicitOrSharedPCHs;
+
+        PublicDependencyModuleNames.AddRange(
+            new string[]
+            {
+                "Core",
+            }
+        );
+
+        PrivateDependencyModuleNames.AddRange(
+            new string[]
+            {
+                "CoreUObject",
+                "Engine",
+                "EditorFramework",
+                "SettingsEditor",
+                "UnrealEd",
+				"Slate",
+				"SlateCore",
+            }
+        );
+    }
+}
\ No newline at end of file
diff --git a/Engine/Source/Editor/InputBindingSettingsViewer/Private/InputBindingSettingsViewerModule.cpp b/Engine/Source/Editor/InputBindingSettingsViewer/Private/InputBindingSettingsViewerModule.cpp
new file mode 100644
--- /dev/null
+++ b/Engine/Source/Editor/InputBindingSettingsViewer/Private/InputBindingSettingsViewerModule.cpp
@@ -0,0 +1,81 @@
+﻿// Copyright Epic Games, Inc. All Rights Reserved.
+
+#include "InputBindingSettingsViewerModule.h"
+
+#include "Framework/Docking/TabManager.h"
+#include "ISettingsCategory.h"
+#include "ISettingsContainer.h"
+#include "ISettingsEditorModel.h"
+#include "ISettingsEditorModule.h"
+#include "ISettingsModule.h"
+#include "Modules/ModuleManager.h"
+#include "Widgets/Docking/SDockTab.h"
+
+#define LOCTEXT_NAMESPACE "InputBindingSettingsViewerModule"
+
+const FLazyName FInputBindingSettingsViewerModule::SettingsTabName("InputBindingSettings");
+const FLazyName FInputBindingSettingsViewerModule::ContainerName("InputBinding");
+
+void FInputBindingSettingsViewerModule::ShowSettings(const FName& InCategoryName, const FName& InSectionName)
+{
+	FGlobalTabmanager::Get()->TryInvokeTab(FTabId(SettingsTabName));
+	if (const ISettingsEditorModelPtr SettingsEditorModel = SettingsEditorModelWeak.Pin())
+	{
+		if (const ISettingsCategoryPtr Category = SettingsEditorModel->GetSettingsContainer()->GetCategory(InCategoryName))
+		{
+			SettingsEditorModel->SelectSection(Category->GetSection(InSectionName));
+		}
+	}
+}
+
+void FInputBindingSettingsViewerModule::StartupModule()
+{
+	if (ISettingsModule* SettingsModule = FModuleManager::GetModulePtr<ISettingsModule>("Settings"))
+	{
+		SettingsModule->RegisterViewer(ContainerName, *this);
+	}
+
+	FGlobalTabmanager::Get()->RegisterNomadTabSpawner(SettingsTabName, FOnSpawnTab::CreateRaw(this, &FInputBindingSettingsViewerModule::HandleSpawnSettingsTab))
+		.SetDisplayName(LOCTEXT("InputBindingSettingsTabTitle", "Input Settings"))
+		.SetMenuType(ETabSpawnerMenuType::Hidden)
+		.SetIcon(FSlateIcon(FAppStyle::GetAppStyleSetName(), "SystemWideCommands.OpenKeyboardShortcuts"));
+}
+
+void FInputBindingSettingsViewerModule::ShutdownModule()
+{
+	FGlobalTabmanager::Get()->UnregisterNomadTabSpawner(SettingsTabName);
+	if (ISettingsModule* SettingsModule = FModuleManager::GetModulePtr<ISettingsModule>("Settings"))
+	{
+		SettingsModule->UnregisterViewer(ContainerName);
+	}
+}
+
+bool FInputBindingSettingsViewerModule::SupportsDynamicReloading()
+{
+	return true;
+}
+
+TSharedRef<SDockTab> FInputBindingSettingsViewerModule::HandleSpawnSettingsTab(const FSpawnTabArgs& InSpawnTabArgs)
+{
+	TSharedRef<SWidget> SettingsEditor = SNullWidget::NullWidget;
+	if (ISettingsModule* SettingsModule = FModuleManager::GetModulePtr<ISettingsModule>("Settings"))
+	{
+		if (const ISettingsContainerPtr SettingsContainer = SettingsModule->GetContainer(ContainerName))
+		{
+			ISettingsEditorModule& SettingsEditorModule = FModuleManager::GetModuleChecked<ISettingsEditorModule>("SettingsEditor");
+			const ISettingsEditorModelRef SettingsEditorModel = SettingsEditorModule.CreateModel(SettingsContainer.ToSharedRef());
+			SettingsEditor = SettingsEditorModule.CreateEditor(SettingsEditorModel);
+			SettingsEditorModelWeak = SettingsEditorModel;
+		}
+	}
+
+	return SNew(SDockTab)
+		.TabRole(ETabRole::NomadTab)
+		[
+			SettingsEditor
+		];
+}
+
+IMPLEMENT_MODULE(FInputBindingSettingsViewerModule, InputBindingSettingsViewer);
+
+#undef LOCTEXT_NAMESPACE
\ No newline at end of file
diff --git a/Engine/Source/Editor/InputBindingSettingsViewer/Private/InputBindingSettingsViewerModule.h b/Engine/Source/Editor/InputBindingSettingsViewer/Private/InputBindingSettingsViewerModule.h
new file mode 100644
--- /dev/null
+++ b/Engine/Source/Editor/InputBindingSettingsViewer/Private/InputBindingSettingsViewerModule.h
@@ -0,0 +1,37 @@
+// Copyright Epic Games, Inc. All Rights Reserved.
+
+#pragma once
+
+#include "ISettingsViewer.h"
+#include "Modules/ModuleInterface.h"
+
+class FSpawnTabArgs;
+class ISettingsEditorModel;
+class SDockTab;
+
+/**
+ * Used for all editor inputs (keyboard shortcuts, ...)
+ */
+class FInputBindingSettingsViewerModule
+	: public IModuleInterface
+	, public ISettingsViewer
+{
+public:
+	static const FLazyName SettingsTabName;
+	static const FLazyName ContainerName;
+
+	//~ Begin ISettingsViewer
+	virtual void ShowSettings(const FName& InCategoryName, const FName& InSectionName) override;
+	//~ End ISettingsViewer
+
+	//~ Begin IModuleInterface
+	virtual void StartupModule() override;
+	virtual void ShutdownModule() override;
+	virtual bool SupportsDynamicReloading() override;
+	//~ End IModuleInterface
+
+private:
+	TSharedRef<SDockTab> HandleSpawnSettingsTab(const FSpawnTabArgs& InSpawnTabArgs);
+
+	TWeakPtr<ISettingsEditorModel> SettingsEditorModelWeak;
+};
diff --git a/Engine/Source/Editor/UnrealEd/Private/EditorEngine.cpp b/Engine/Source/Editor/UnrealEd/Private/EditorEngine.cpp
--- a/Engine/Source/Editor/UnrealEd/Private/EditorEngine.cpp
+++ b/Engine/Source/Editor/UnrealEd/Private/EditorEngine.cpp
@@ -1515,6 +1515,7 @@ void UEditorEngine::LoadDefaultEditorModules()
 			TEXT("LegacyProjectLauncher"),
 			TEXT("SettingsEditor"),
 			TEXT("EditorSettingsViewer"),
+			TEXT("InputBindingSettingsViewer"),
 			TEXT("ProjectSettingsViewer"),
 			TEXT("Blutility"),
 			TEXT("ScriptableEditorWidgets"),
diff --git a/Engine/Source/Editor/UnrealEd/Private/Toolkits/GlobalEditorCommonCommands.cpp b/Engine/Source/Editor/UnrealEd/Private/Toolkits/GlobalEditorCommonCommands.cpp
--- a/Engine/Source/Editor/UnrealEd/Private/Toolkits/GlobalEditorCommonCommands.cpp
+++ b/Engine/Source/Editor/UnrealEd/Private/Toolkits/GlobalEditorCommonCommands.cpp
@@ -149,7 +149,7 @@ void FGlobalEditorCommonCommands::MapActions(TSharedRef<FUICommandList>& Toolkit
 
 	ToolkitCommands->MapAction(
 		GlobalEditorCommands.OpenKeyboardShortcuts,
-		FExecuteAction::CreateStatic(&FGlobalEditorCommonCommands::OpenSettings, FName("Editor"), FName("Shortcuts"), FName("InputBindings")));
+		FExecuteAction::CreateStatic(&FGlobalEditorCommonCommands::OpenSettings, FName("InputBinding"), FName("Shortcuts"), FName("InputBindings")));
 
 	CategoryName = NAME_None; // "All settings"
 	SectionName = NAME_None; // "All settings"
diff --git a/Engine/Source/Editor/UnrealEd/UnrealEd.Build.cs b/Engine/Source/Editor/UnrealEd/UnrealEd.Build.cs
--- a/Engine/Source/Editor/UnrealEd/UnrealEd.Build.cs
+++ b/Engine/Source/Editor/UnrealEd/UnrealEd.Build.cs
@@ -210,6 +210,7 @@ public class UnrealEd : ModuleRules
 				"CurveTableEditor",
 				"DataTableEditor",
 				"EditorSettingsViewer",
+				"InputBindingSettingsViewer",
 				"LandscapeEditor",
 				"KismetCompiler",
 				"DetailCustomizations",
